Fix -ffast-math in SQLite regarding sqlite3IsNaN by using the math.h version of isnan

diff --git a/Engine/Plugins/Runtime/Database/SQLiteCore/Source/SQLiteCore/SQLiteCore.Build.cs b/Engine/Plugins/Runtime/Database/SQLiteCore/Source/SQLiteCore/SQLiteCore.Build.cs
--- a/Engine/Plugins/Runtime/Database/SQLiteCore/Source/SQLiteCore/SQLiteCore.Build.cs
+++ b/Engine/Plugins/Runtime/Database/SQLiteCore/Source/SQLiteCore/SQLiteCore.Build.cs
@@ -20,6 +20,9 @@ namespace UnrealBuildTool.Rules
 				PrivateDefinitions.Add("SQLITE_THREADSAFE=0"); // No threading on HTML5
 			}
 
+			// Use the math.h version of isnan rather than the SQLite version to avoid a -ffast-math error
+			PrivateDefinitions.Add("SQLITE_HAVE_ISNAN=1");
+
 			// Enable SQLite debug checks?
 			//PrivateDefinitions.Add("SQLITE_DEBUG");
 
