Merging using //UE4/Partner-Microsoft-XR/_to_//UE4/Release-4.23/

Allow for SQLite in HoloLens builds

diff --git a/Engine/Plugins/Runtime/Database/SQLiteCore/SQLiteCore.uplugin b/Engine/Plugins/Runtime/Database/SQLiteCore/SQLiteCore.uplugin
--- a/Engine/Plugins/Runtime/Database/SQLiteCore/SQLiteCore.uplugin
+++ b/Engine/Plugins/Runtime/Database/SQLiteCore/SQLiteCore.uplugin
@@ -19,10 +19,7 @@
 		{
 			"Name": "SQLiteCore",
 			"Type": "Runtime",
-			"LoadingPhase": "PreDefault",
-			"BlacklistPlatforms": [
-				"HoloLens"
-			]
-		},
+			"LoadingPhase": "PreDefault"
+		}
 	]
 }
\ No newline at end of file
diff --git a/Engine/Plugins/Runtime/Database/SQLiteSupport/SQLiteSupport.uplugin b/Engine/Plugins/Runtime/Database/SQLiteSupport/SQLiteSupport.uplugin
--- a/Engine/Plugins/Runtime/Database/SQLiteSupport/SQLiteSupport.uplugin
+++ b/Engine/Plugins/Runtime/Database/SQLiteSupport/SQLiteSupport.uplugin
@@ -18,10 +18,7 @@
 		{
 			"Name": "SQLiteSupport",
 			"Type": "Runtime",
-			"LoadingPhase": "Default",
-			"BlacklistPlatforms": [
-				"HoloLens"
-			]
+			"LoadingPhase": "Default"
 		}
 	],
 	"Plugins": [
