Added a couple of missing #includes to the SQLiteCore plugin code to solve compile errors when building with fastbuild

diff --git a/Engine/Plugins/Runtime/Database/SQLiteCore/Source/SQLiteCore/Private/SQLiteEmbedded.c b/Engine/Plugins/Runtime/Database/SQLiteCore/Source/SQLiteCore/Private/SQLiteEmbedded.c
--- a/Engine/Plugins/Runtime/Database/SQLiteCore/Source/SQLiteCore/Private/SQLiteEmbedded.c
+++ b/Engine/Plugins/Runtime/Database/SQLiteCore/Source/SQLiteCore/Private/SQLiteEmbedded.c
@@ -9,5 +9,7 @@
 
 /** Include SQLite, but not if we're building for analysis as the code emits warnings */
 #if !USING_CODE_ANALYSIS
+#include "HAL/PreprocessorHelpers.h"
+#include COMPILED_PLATFORM_HEADER(PlatformCompilerPreSetup.h)
 #include "sqlite/sqlite3.inl"
 #endif
