CIS fix

#rnx

diff --git a/Engine/Plugins/Runtime/Database/SQLiteCore/Source/SQLiteCore/Private/SQLiteDatabase.cpp b/Engine/Plugins/Runtime/Database/SQLiteCore/Source/SQLiteCore/Private/SQLiteDatabase.cpp
--- a/Engine/Plugins/Runtime/Database/SQLiteCore/Source/SQLiteCore/Private/SQLiteDatabase.cpp
+++ b/Engine/Plugins/Runtime/Database/SQLiteCore/Source/SQLiteCore/Private/SQLiteDatabase.cpp
@@ -123,7 +123,7 @@ bool FSQLiteDatabase::GetApplicationId(int32& OutApplicationId) const
 	{
 		InStatement.GetColumnValueByIndex(0, OutApplicationId);
 		return ESQLitePreparedStatementExecuteRowResult::Stop;
-	});
+	}) == 1;
 }
 
 bool FSQLiteDatabase::SetApplicationId(const int32 InApplicationId)
@@ -137,7 +137,7 @@ bool FSQLiteDatabase::GetUserVersion(int32& OutUserVersion) const
 	{
 		InStatement.GetColumnValueByIndex(0, OutUserVersion);
 		return ESQLitePreparedStatementExecuteRowResult::Stop;
-	});
+	}) == 1;
 }
 
 bool FSQLiteDatabase::SetUserVersion(const int32 InUserVersion)
