Updating SQLite pull to newest release.  Enabling additional #defines to allow for JSON support.

diff --git a/Engine/Plugins/Runtime/Database/SQLiteCore/Source/SQLiteCore/Private/sqlite/sqlite3.h b/Engine/Plugins/Runtime/Database/SQLiteCore/Source/SQLiteCore/Private/sqlite/sqlite3.h
--- a/Engine/Plugins/Runtime/Database/SQLiteCore/Source/SQLiteCore/Private/sqlite/sqlite3.h
+++ b/Engine/Plugins/Runtime/Database/SQLiteCore/Source/SQLiteCore/Private/sqlite/sqlite3.h
@@ -123,9 +123,9 @@ extern "C" {
 ** [sqlite3_libversion_number()], [sqlite3_sourceid()],
 ** [sqlite_version()] and [sqlite_source_id()].
 */
-#define SQLITE_VERSION        "3.28.0"
-#define SQLITE_VERSION_NUMBER 3028000
-#define SQLITE_SOURCE_ID      "2019-04-16 19:49:53 884b4b7e502b4e991677b53971277adfaf0a04a284f8e483e2553d0f83156b50"
+#define SQLITE_VERSION        "3.31.1"
+#define SQLITE_VERSION_NUMBER 3031001
+#define SQLITE_SOURCE_ID      "2020-01-27 19:55:54 3bfa9cc97da10598521b342961df8f5f68c7388fa117345eeb516eaa837bb4d6"
 
 /*
 ** CAPI3REF: Run-Time Library Version Numbers
@@ -516,6 +516,7 @@ SQLITE_API int sqlite3_exec(
 #define SQLITE_CANTOPEN_FULLPATH       (SQLITE_CANTOPEN | (3<<8))
 #define SQLITE_CANTOPEN_CONVPATH       (SQLITE_CANTOPEN | (4<<8))
 #define SQLITE_CANTOPEN_DIRTYWAL       (SQLITE_CANTOPEN | (5<<8)) /* Not Used */
+#define SQLITE_CANTOPEN_SYMLINK        (SQLITE_CANTOPEN | (6<<8))
 #define SQLITE_CORRUPT_VTAB            (SQLITE_CORRUPT | (1<<8))
 #define SQLITE_CORRUPT_SEQUENCE        (SQLITE_CORRUPT | (2<<8))
 #define SQLITE_READONLY_RECOVERY       (SQLITE_READONLY | (1<<8))
@@ -535,11 +536,13 @@ SQLITE_API int sqlite3_exec(
 #define SQLITE_CONSTRAINT_UNIQUE       (SQLITE_CONSTRAINT | (8<<8))
 #define SQLITE_CONSTRAINT_VTAB         (SQLITE_CONSTRAINT | (9<<8))
 #define SQLITE_CONSTRAINT_ROWID        (SQLITE_CONSTRAINT |(10<<8))
+#define SQLITE_CONSTRAINT_PINNED       (SQLITE_CONSTRAINT |(11<<8))
 #define SQLITE_NOTICE_RECOVER_WAL      (SQLITE_NOTICE | (1<<8))
 #define SQLITE_NOTICE_RECOVER_ROLLBACK (SQLITE_NOTICE | (2<<8))
 #define SQLITE_WARNING_AUTOINDEX       (SQLITE_WARNING | (1<<8))
 #define SQLITE_AUTH_USER               (SQLITE_AUTH | (1<<8))
 #define SQLITE_OK_LOAD_PERMANENTLY     (SQLITE_OK | (1<<8))
+#define SQLITE_OK_SYMLINK              (SQLITE_OK | (2<<8))
 
 /*
 ** CAPI3REF: Flags For File Open Operations
@@ -568,6 +571,7 @@ SQLITE_API int sqlite3_exec(
 #define SQLITE_OPEN_SHAREDCACHE      0x00020000  /* Ok for sqlite3_open_v2() */
 #define SQLITE_OPEN_PRIVATECACHE     0x00040000  /* Ok for sqlite3_open_v2() */
 #define SQLITE_OPEN_WAL              0x00080000  /* VFS only */
+#define SQLITE_OPEN_NOFOLLOW         0x01000000  /* Ok for sqlite3_open_v2() */
 
 /* Reserved:                         0x00F00000 */
 
@@ -979,16 +983,16 @@ struct sqlite3_io_methods {
 ** ^The [SQLITE_FCNTL_BUSYHANDLER]
 ** file-control may be invoked by SQLite on the database file handle
 ** shortly after it is opened in order to provide a custom VFS with access
-** to the connections busy-handler callback. The argument is of type (void **)
+** to the connection's busy-handler callback. The argument is of type (void**)
 ** - an array of two (void *) values. The first (void *) actually points
-** to a function of type (int (*)(void *)). In order to invoke the connections
+** to a function of type (int (*)(void *)). In order to invoke the connection's
 ** busy-handler, this function should be invoked with the second (void *) in
 ** the array as the only argument. If it returns non-zero, then the operation
 ** should be retried. If it returns zero, the custom VFS should abandon the
 ** current operation.
 **
 ** <li>[[SQLITE_FCNTL_TEMPFILENAME]]
-** ^Application can invoke the [SQLITE_FCNTL_TEMPFILENAME] file-control
+** ^Applications can invoke the [SQLITE_FCNTL_TEMPFILENAME] file-control
 ** to have SQLite generate a
 ** temporary filename using the same algorithm that is followed to generate
 ** temporary filenames for TEMP tables and other internal uses.  The
@@ -1101,12 +1105,18 @@ struct sqlite3_io_methods {
 ** not provide a mechanism to detect changes to MAIN only.  Also, the
 ** [sqlite3_total_changes()] interface responds to internal changes only and
 ** omits changes made by other database connections.  The
-** [PRAGMA data_version] command provide a mechanism to detect changes to
+** [PRAGMA data_version] command provides a mechanism to detect changes to
 ** a single attached database that occur due to other database connections,
 ** but omits changes implemented by the database connection on which it is
 ** called.  This file control is the only mechanism to detect changes that
 ** happen either internally or externally and that are associated with
 ** a particular attached database.
+**
+** <li>[[SQLITE_FCNTL_CKPT_DONE]]
+** The [SQLITE_FCNTL_CKPT_DONE] opcode is invoked from within a checkpoint
+** in wal mode after the client has finished copying pages from the wal
+** file to the database file, but before the *-shm file is updated to
+** record the fact that the pages have been checkpointed.
 ** </ul>
 */
 #define SQLITE_FCNTL_LOCKSTATE               1
@@ -1144,6 +1154,7 @@ struct sqlite3_io_methods {
 #define SQLITE_FCNTL_LOCK_TIMEOUT           34
 #define SQLITE_FCNTL_DATA_VERSION           35
 #define SQLITE_FCNTL_SIZE_LIMIT             36
+#define SQLITE_FCNTL_CKPT_DONE              37
 
 /* deprecated names */
 #define SQLITE_GET_LOCKPROXYFILE      SQLITE_FCNTL_GET_LOCKPROXYFILE
@@ -1189,10 +1200,10 @@ typedef struct sqlite3_api_routines sqlite3_api_routines;
 ** to 3 with SQLite [version 3.7.6] on [dateof:3.7.6].  Additional fields
 ** may be appended to the sqlite3_vfs object and the iVersion value
 ** may increase again in future versions of SQLite.
-** Note that the structure
-** of the sqlite3_vfs object changes in the transition from
+** Note that due to an oversight, the structure
+** of the sqlite3_vfs object changed in the transition from
 ** SQLite [version 3.5.9] to [version 3.6.0] on [dateof:3.6.0]
-** and yet the iVersion field was not modified.
+** and yet the iVersion field was not increased.
 **
 ** The szOsFile field is the size of the subclassed [sqlite3_file]
 ** structure used by this VFS.  mxPathname is the maximum length of
@@ -1283,7 +1294,7 @@ typedef struct sqlite3_api_routines sqlite3_api_routines;
 ** for exclusive access.
 **
 ** ^At least szOsFile bytes of memory are allocated by SQLite
-** to hold the  [sqlite3_file] structure passed as the third
+** to hold the [sqlite3_file] structure passed as the third
 ** argument to xOpen.  The xOpen method does not have to
 ** allocate the structure; it should just fill it in.  Note that
 ** the xOpen method must set the sqlite3_file.pMethods to either
@@ -1296,8 +1307,14 @@ typedef struct sqlite3_api_routines sqlite3_api_routines;
 ** ^The flags argument to xAccess() may be [SQLITE_ACCESS_EXISTS]
 ** to test for the existence of a file, or [SQLITE_ACCESS_READWRITE] to
 ** test whether a file is readable and writable, or [SQLITE_ACCESS_READ]
-** to test whether a file is at least readable.   The file can be a
-** directory.
+** to test whether a file is at least readable.  The SQLITE_ACCESS_READ
+** flag is never actually used and is not implemented in the built-in
+** VFSes of SQLite.  The file is named by the second argument and can be a
+** directory. The xAccess method returns [SQLITE_OK] on success or some
+** non-zero error code if there is an I/O error or if the name of
+** the file given in the second argument is illegal.  If SQLITE_OK
+** is returned, then non-zero or zero is written into *pResOut to indicate
+** whether or not the file is accessible.  
 **
 ** ^SQLite will always allocate at least mxPathname+1 bytes for the
 ** output buffer xFullPathname.  The exact size of the output buffer
@@ -1614,7 +1631,7 @@ SQLITE_API int sqlite3_db_config(sqlite3*, int op, ...);
 ** that causes the corresponding memory allocation to fail.
 **
 ** The xInit method initializes the memory allocator.  For example,
-** it might allocate any require mutexes or initialize internal data
+** it might allocate any required mutexes or initialize internal data
 ** structures.  The xShutdown method is invoked (indirectly) by
 ** [sqlite3_shutdown()] and should deallocate any resources acquired
 ** by xInit.  The pAppData pointer is used as the only parameter to
@@ -1736,6 +1753,7 @@ struct sqlite3_mem_methods {
 ** memory allocation statistics. ^(When memory allocation statistics are
 ** disabled, the following SQLite interfaces become non-operational:
 **   <ul>
+**   <li> [sqlite3_hard_heap_limit64()]
 **   <li> [sqlite3_memory_used()]
 **   <li> [sqlite3_memory_highwater()]
 **   <li> [sqlite3_soft_heap_limit64()]
@@ -1754,7 +1772,7 @@ struct sqlite3_mem_methods {
 ** <dd> ^The SQLITE_CONFIG_PAGECACHE option specifies a memory pool
 ** that SQLite can use for the database page cache with the default page
 ** cache implementation.  
-** This configuration option is a no-op if an application-define page
+** This configuration option is a no-op if an application-defined page
 ** cache implementation is loaded using the [SQLITE_CONFIG_PCACHE2].
 ** ^There are three arguments to SQLITE_CONFIG_PAGECACHE: A pointer to
 ** 8-byte aligned memory (pMem), the size of each page cache line (sz),
@@ -2087,6 +2105,17 @@ struct sqlite3_mem_methods {
 ** following this call.  The second parameter may be a NULL pointer, in
 ** which case the trigger setting is not reported back. </dd>
 **
+** [[SQLITE_DBCONFIG_ENABLE_VIEW]]
+** <dt>SQLITE_DBCONFIG_ENABLE_VIEW</dt>
+** <dd> ^This option is used to enable or disable [CREATE VIEW | views].
+** There should be two additional arguments.
+** The first argument is an integer which is 0 to disable views,
+** positive to enable views or negative to leave the setting unchanged.
+** The second parameter is a pointer to an integer into which
+** is written 0 or 1 to indicate whether views are disabled or enabled
+** following this call.  The second parameter may be a NULL pointer, in
+** which case the view setting is not reported back. </dd>
+**
 ** [[SQLITE_DBCONFIG_ENABLE_FTS3_TOKENIZER]]
 ** <dt>SQLITE_DBCONFIG_ENABLE_FTS3_TOKENIZER</dt>
 ** <dd> ^This option is used to enable or disable the
@@ -2198,6 +2227,7 @@ struct sqlite3_mem_methods {
 ** features include but are not limited to the following:
 ** <ul>
 ** <li> The [PRAGMA writable_schema=ON] statement.
+** <li> The [PRAGMA journal_mode=OFF] statement.
 ** <li> Writes to the [sqlite_dbpage] virtual table.
 ** <li> Direct writes to [shadow tables].
 ** </ul>
@@ -2213,6 +2243,77 @@ struct sqlite3_mem_methods {
 ** integer into which is written 0 or 1 to indicate whether the writable_schema
 ** is enabled or disabled following this call.
 ** </dd>
+**
+** [[SQLITE_DBCONFIG_LEGACY_ALTER_TABLE]]
+** <dt>SQLITE_DBCONFIG_LEGACY_ALTER_TABLE</dt>
+** <dd>The SQLITE_DBCONFIG_LEGACY_ALTER_TABLE option activates or deactivates
+** the legacy behavior of the [ALTER TABLE RENAME] command such it
+** behaves as it did prior to [version 3.24.0] (2018-06-04).  See the
+** "Compatibility Notice" on the [ALTER TABLE RENAME documentation] for
+** additional information. This feature can also be turned on and off
+** using the [PRAGMA legacy_alter_table] statement.
+** </dd>
+**
+** [[SQLITE_DBCONFIG_DQS_DML]]
+** <dt>SQLITE_DBCONFIG_DQS_DML</td>
+** <dd>The SQLITE_DBCONFIG_DQS_DML option activates or deactivates
+** the legacy [double-quoted string literal] misfeature for DML statements
+** only, that is DELETE, INSERT, SELECT, and UPDATE statements. The
+** default value of this setting is determined by the [-DSQLITE_DQS]
+** compile-time option.
+** </dd>
+**
+** [[SQLITE_DBCONFIG_DQS_DDL]]
+** <dt>SQLITE_DBCONFIG_DQS_DDL</td>
+** <dd>The SQLITE_DBCONFIG_DQS option activates or deactivates
+** the legacy [double-quoted string literal] misfeature for DDL statements,
+** such as CREATE TABLE and CREATE INDEX. The
+** default value of this setting is determined by the [-DSQLITE_DQS]
+** compile-time option.
+** </dd>
+**
+** [[SQLITE_DBCONFIG_TRUSTED_SCHEMA]]
+** <dt>SQLITE_DBCONFIG_TRUSTED_SCHEMA</td>
+** <dd>The SQLITE_DBCONFIG_TRUSTED_SCHEMA option tells SQLite to
+** assume that database schemas (the contents of the [sqlite_master] tables)
+** are untainted by malicious content.
+** When the SQLITE_DBCONFIG_TRUSTED_SCHEMA option is disabled, SQLite
+** takes additional defensive steps to protect the application from harm
+** including:
+** <ul>
+** <li> Prohibit the use of SQL functions inside triggers, views,
+** CHECK constraints, DEFAULT clauses, expression indexes, 
+** partial indexes, or generated columns
+** unless those functions are tagged with [SQLITE_INNOCUOUS].
+** <li> Prohibit the use of virtual tables inside of triggers or views
+** unless those virtual tables are tagged with [SQLITE_VTAB_INNOCUOUS].
+** </ul>
+** This setting defaults to "on" for legacy compatibility, however
+** all applications are advised to turn it off if possible. This setting
+** can also be controlled using the [PRAGMA trusted_schema] statement.
+** </dd>
+**
+** [[SQLITE_DBCONFIG_LEGACY_FILE_FORMAT]]
+** <dt>SQLITE_DBCONFIG_LEGACY_FILE_FORMAT</td>
+** <dd>The SQLITE_DBCONFIG_LEGACY_FILE_FORMAT option activates or deactivates
+** the legacy file format flag.  When activated, this flag causes all newly
+** created database file to have a schema format version number (the 4-byte
+** integer found at offset 44 into the database header) of 1.  This in turn
+** means that the resulting database file will be readable and writable by
+** any SQLite version back to 3.0.0 ([dateof:3.0.0]).  Without this setting,
+** newly created databases are generally not understandable by SQLite versions
+** prior to 3.3.0 ([dateof:3.3.0]).  As these words are written, there
+** is now scarcely any need to generated database files that are compatible 
+** all the way back to version 3.0.0, and so this setting is of little
+** practical use, but is provided so that SQLite can continue to claim the
+** ability to generate new database files that are compatible with  version
+** 3.0.0.
+** <p>Note that when the SQLITE_DBCONFIG_LEGACY_FILE_FORMAT setting is on,
+** the [VACUUM] command will fail with an obscure error when attempting to
+** process a table with generated columns and a descending index.  This is
+** not considered a bug since SQLite versions 3.3.0 and earlier do not support
+** either generated columns or decending indexes.
+** </dd>
 ** </dl>
 */
 #define SQLITE_DBCONFIG_MAINDBNAME            1000 /* const char* */
@@ -2227,7 +2328,13 @@ struct sqlite3_mem_methods {
 #define SQLITE_DBCONFIG_RESET_DATABASE        1009 /* int int* */
 #define SQLITE_DBCONFIG_DEFENSIVE             1010 /* int int* */
 #define SQLITE_DBCONFIG_WRITABLE_SCHEMA       1011 /* int int* */
-#define SQLITE_DBCONFIG_MAX                   1011 /* Largest DBCONFIG */
+#define SQLITE_DBCONFIG_LEGACY_ALTER_TABLE    1012 /* int int* */
+#define SQLITE_DBCONFIG_DQS_DML               1013 /* int int* */
+#define SQLITE_DBCONFIG_DQS_DDL               1014 /* int int* */
+#define SQLITE_DBCONFIG_ENABLE_VIEW           1015 /* int int* */
+#define SQLITE_DBCONFIG_LEGACY_FILE_FORMAT    1016 /* int int* */
+#define SQLITE_DBCONFIG_TRUSTED_SCHEMA        1017 /* int int* */
+#define SQLITE_DBCONFIG_MAX                   1017 /* Largest DBCONFIG */
 
 /*
 ** CAPI3REF: Enable Or Disable Extended Result Codes
@@ -2433,7 +2540,7 @@ SQLITE_API int sqlite3_total_changes(sqlite3*);
 ** ^The sqlite3_interrupt(D) call is in effect until all currently running
 ** SQL statements on [database connection] D complete.  ^Any new SQL statements
 ** that are started after the sqlite3_interrupt() call and before the 
-** running statements reaches zero are interrupted as if they had been
+** running statement count reaches zero are interrupted as if they had been
 ** running prior to the sqlite3_interrupt() call.  ^New SQL statements
 ** that are started after the running statement count reaches zero are
 ** not effected by the sqlite3_interrupt().
@@ -2601,9 +2708,9 @@ SQLITE_API int sqlite3_busy_timeout(sqlite3*, int ms);
 **        Cindy       | 21
 ** </pre></blockquote>
 **
-** There are two column (M==2) and three rows (N==3).  Thus the
+** There are two columns (M==2) and three rows (N==3).  Thus the
 ** result table has 8 entries.  Suppose the result table is stored
-** in an array names azResult.  Then azResult holds this content:
+** in an array named azResult.  Then azResult holds this content:
 **
 ** <blockquote><pre>
 **        azResult&#91;0] = "Name";
@@ -2696,7 +2803,7 @@ SQLITE_API char *sqlite3_vsnprintf(int,char*,const char*, va_list);
 **
 ** The SQLite core uses these three routines for all of its own
 ** internal memory allocation needs. "Core" in the previous sentence
-** does not include operating-system specific VFS implementation.  The
+** does not include operating-system specific [VFS] implementation.  The
 ** Windows VFS uses native malloc() and free() for some operations.
 **
 ** ^The sqlite3_malloc() routine returns a pointer to a block
@@ -2757,19 +2864,6 @@ SQLITE_API char *sqlite3_vsnprintf(int,char*,const char*, va_list);
 ** 4 byte boundary if the [SQLITE_4_BYTE_ALIGNED_MALLOC] compile-time
 ** option is used.
 **
-** In SQLite version 3.5.0 and 3.5.1, it was possible to define
-** the SQLITE_OMIT_MEMORY_ALLOCATION which would cause the built-in
-** implementation of these routines to be omitted.  That capability
-** is no longer provided.  Only built-in memory allocators can be used.
-**
-** Prior to SQLite version 3.7.10, the Windows OS interface layer called
-** the system malloc() and free() directly when converting
-** filenames between the UTF-8 encoding used by SQLite
-** and whatever filename encoding is used by the particular Windows
-** installation.  Memory allocation errors were detected, but
-** they were reported back as [SQLITE_CANTOPEN] or
-** [SQLITE_IOERR] rather than [SQLITE_NOMEM].
-**
 ** The pointer arguments to [sqlite3_free()] and [sqlite3_realloc()]
 ** must be either NULL or else pointers obtained from a prior
 ** invocation of [sqlite3_malloc()] or [sqlite3_realloc()] that have
@@ -2818,7 +2912,7 @@ SQLITE_API sqlite3_int64 sqlite3_memory_highwater(int resetFlag);
 ** SQLite contains a high-quality pseudo-random number generator (PRNG) used to
 ** select random [ROWID | ROWIDs] when inserting new records into a table that
 ** already uses the largest possible [ROWID].  The PRNG is also used for
-** the build-in random() and randomblob() SQL functions.  This interface allows
+** the built-in random() and randomblob() SQL functions.  This interface allows
 ** applications to access the same PRNG for other purposes.
 **
 ** ^A call to this routine stores N bytes of randomness into buffer P.
@@ -3192,10 +3286,8 @@ SQLITE_API void sqlite3_progress_handler(sqlite3*, int, int(*)(void*), void*);
 ** The sqlite3_open_v2() interface works like sqlite3_open()
 ** except that it accepts two additional parameters for additional control
 ** over the new database connection.  ^(The flags parameter to
-** sqlite3_open_v2() can take one of
-** the following three values, optionally combined with the 
-** [SQLITE_OPEN_NOMUTEX], [SQLITE_OPEN_FULLMUTEX], [SQLITE_OPEN_SHAREDCACHE],
-** [SQLITE_OPEN_PRIVATECACHE], and/or [SQLITE_OPEN_URI] flags:)^
+** sqlite3_open_v2() must include, at a minimum, one of the following
+** three flag combinations:)^
 **
 ** <dl>
 ** ^(<dt>[SQLITE_OPEN_READONLY]</dt>
@@ -3213,23 +3305,51 @@ SQLITE_API void sqlite3_progress_handler(sqlite3*, int, int(*)(void*), void*);
 ** sqlite3_open() and sqlite3_open16().</dd>)^
 ** </dl>
 **
+** In addition to the required flags, the following optional flags are
+** also supported:
+**
+** <dl>
+** ^(<dt>[SQLITE_OPEN_URI]</dt>
+** <dd>The filename can be interpreted as a URI if this flag is set.</dd>)^
+**
+** ^(<dt>[SQLITE_OPEN_MEMORY]</dt>
+** <dd>The database will be opened as an in-memory database.  The database
+** is named by the "filename" argument for the purposes of cache-sharing,
+** if shared cache mode is enabled, but the "filename" is otherwise ignored.
+** </dd>)^
+**
+** ^(<dt>[SQLITE_OPEN_NOMUTEX]</dt>
+** <dd>The new database connection will use the "multi-thread"
+** [threading mode].)^  This means that separate threads are allowed
+** to use SQLite at the same time, as long as each thread is using
+** a different [database connection].
+**
+** ^(<dt>[SQLITE_OPEN_FULLMUTEX]</dt>
+** <dd>The new database connection will use the "serialized"
+** [threading mode].)^  This means the multiple threads can safely
+** attempt to use the same database connection at the same time.
+** (Mutexes will block any actual concurrency, but in this mode
+** there is no harm in trying.)
+**
+** ^(<dt>[SQLITE_OPEN_SHAREDCACHE]</dt>
+** <dd>The database is opened [shared cache] enabled, overriding
+** the default shared cache setting provided by
+** [sqlite3_enable_shared_cache()].)^
+**
+** ^(<dt>[SQLITE_OPEN_PRIVATECACHE]</dt>
+** <dd>The database is opened [shared cache] disabled, overriding
+** the default shared cache setting provided by
+** [sqlite3_enable_shared_cache()].)^
+**
+** [[OPEN_NOFOLLOW]] ^(<dt>[SQLITE_OPEN_NOFOLLOW]</dt>
+** <dd>The database filename is not allowed to be a symbolic link</dd>
+** </dl>)^
+**
 ** If the 3rd parameter to sqlite3_open_v2() is not one of the
-** combinations shown above optionally combined with other
+** required combinations shown above optionally combined with other
 ** [SQLITE_OPEN_READONLY | SQLITE_OPEN_* bits]
 ** then the behavior is undefined.
 **
-** ^If the [SQLITE_OPEN_NOMUTEX] flag is set, then the database connection
-** opens in the multi-thread [threading mode] as long as the single-thread
-** mode has not been set at compile-time or start-time.  ^If the
-** [SQLITE_OPEN_FULLMUTEX] flag is set then the database connection opens
-** in the serialized [threading mode] unless single-thread was
-** previously selected at compile-time or start-time.
-** ^The [SQLITE_OPEN_SHAREDCACHE] flag causes the database connection to be
-** eligible to use [shared cache mode], regardless of whether or not shared
-** cache is enabled using [sqlite3_enable_shared_cache()].  ^The
-** [SQLITE_OPEN_PRIVATECACHE] flag causes the database connection to not
-** participate in [shared cache mode] even if it is enabled.
-**
 ** ^The fourth parameter to sqlite3_open_v2() is the name of the
 ** [sqlite3_vfs] object that defines the operating system interface that
 ** the new database connection should use.  ^If the fourth parameter is
@@ -3409,17 +3529,16 @@ SQLITE_API int sqlite3_open_v2(
 /*
 ** CAPI3REF: Obtain Values For URI Parameters
 **
-** These are utility routines, useful to VFS implementations, that check
-** to see if a database file was a URI that contained a specific query 
+** These are utility routines, useful to [VFS|custom VFS implementations],
+** that check if a database file was a URI that contained a specific query 
 ** parameter, and if so obtains the value of that query parameter.
 **
 ** If F is the database filename pointer passed into the xOpen() method of 
-** a VFS implementation when the flags parameter to xOpen() has one or 
-** more of the [SQLITE_OPEN_URI] or [SQLITE_OPEN_MAIN_DB] bits set and
-** P is the name of the query parameter, then
+** a VFS implementation or it is the return value of [sqlite3_db_filename()]
+** and if P is the name of the query parameter, then
 ** sqlite3_uri_parameter(F,P) returns the value of the P
 ** parameter if it exists or a NULL pointer if P does not appear as a 
-** query parameter on F.  If P is a query parameter of F
+** query parameter on F.  If P is a query parameter of F and it
 ** has no explicit value, then sqlite3_uri_parameter(F,P) returns
 ** a pointer to an empty string.
 **
@@ -3431,25 +3550,72 @@ SQLITE_API int sqlite3_open_v2(
 ** sqlite3_uri_boolean(F,P,B) routines returns false (0) if the value of
 ** query parameter P is one of "no", "false", or "off" in any case or
 ** if the value begins with a numeric zero.  If P is not a query
-** parameter on F or if the value of P is does not match any of the
+** parameter on F or if the value of P does not match any of the
 ** above, then sqlite3_uri_boolean(F,P,B) returns (B!=0).
 **
 ** The sqlite3_uri_int64(F,P,D) routine converts the value of P into a
 ** 64-bit signed integer and returns that integer, or D if P does not
 ** exist.  If the value of P is something other than an integer, then
 ** zero is returned.
+**
+** The sqlite3_uri_key(F,N) returns a pointer to the name (not
+** the value) of the N-th query parameter for filename F, or a NULL
+** pointer if N is less than zero or greater than the number of query
+** parameters minus 1.  The N value is zero-based so N should be 0 to obtain
+** the name of the first query parameter, 1 for the second parameter, and
+** so forth.
 ** 
 ** If F is a NULL pointer, then sqlite3_uri_parameter(F,P) returns NULL and
 ** sqlite3_uri_boolean(F,P,B) returns B.  If F is not a NULL pointer and
-** is not a database file pathname pointer that SQLite passed into the xOpen
-** VFS method, then the behavior of this routine is undefined and probably
-** undesirable.
+** is not a database file pathname pointer that the SQLite core passed
+** into the xOpen VFS method, then the behavior of this routine is undefined
+** and probably undesirable.
+**
+** Beginning with SQLite [version 3.31.0] ([dateof:3.31.0]) the input F
+** parameter can also be the name of a rollback journal file or WAL file
+** in addition to the main database file.  Prior to version 3.31.0, these
+** routines would only work if F was the name of the main database file.
+** When the F parameter is the name of the rollback journal or WAL file,
+** it has access to all the same query parameters as were found on the
+** main database file.
 **
 ** See the [URI filename] documentation for additional information.
 */
 SQLITE_API const char *sqlite3_uri_parameter(const char *zFilename, const char *zParam);
 SQLITE_API int sqlite3_uri_boolean(const char *zFile, const char *zParam, int bDefault);
 SQLITE_API sqlite3_int64 sqlite3_uri_int64(const char*, const char*, sqlite3_int64);
+SQLITE_API const char *sqlite3_uri_key(const char *zFilename, int N);
+
+/*
+** CAPI3REF:  Translate filenames
+**
+** These routines are available to [VFS|custom VFS implementations] for
+** translating filenames between the main database file, the journal file,
+** and the WAL file.
+**
+** If F is the name of an sqlite database file, journal file, or WAL file
+** passed by the SQLite core into the VFS, then sqlite3_filename_database(F)
+** returns the name of the corresponding database file.
+**
+** If F is the name of an sqlite database file, journal file, or WAL file
+** passed by the SQLite core into the VFS, or if F is a database filename
+** obtained from [sqlite3_db_filename()], then sqlite3_filename_journal(F)
+** returns the name of the corresponding rollback journal file.
+**
+** If F is the name of an sqlite database file, journal file, or WAL file
+** that was passed by the SQLite core into the VFS, or if F is a database
+** filename obtained from [sqlite3_db_filename()], then
+** sqlite3_filename_wal(F) returns the name of the corresponding
+** WAL file.
+**
+** In all of the above, if F is not the name of a database, journal or WAL
+** filename passed into the VFS from the SQLite core and F is not the
+** return value from [sqlite3_db_filename()], then the result is
+** undefined and is likely a memory access violation.
+*/
+SQLITE_API const char *sqlite3_filename_database(const char*);
+SQLITE_API const char *sqlite3_filename_journal(const char*);
+SQLITE_API const char *sqlite3_filename_wal(const char*);
 
 
 /*
@@ -3768,15 +3934,15 @@ SQLITE_API int sqlite3_limit(sqlite3*, int id, int newVal);
 ** </li>
 **
 ** <li>
-** ^If the specific value bound to [parameter | host parameter] in the 
+** ^If the specific value bound to a [parameter | host parameter] in the 
 ** WHERE clause might influence the choice of query plan for a statement,
 ** then the statement will be automatically recompiled, as if there had been 
-** a schema change, on the first  [sqlite3_step()] call following any change
+** a schema change, on the first [sqlite3_step()] call following any change
 ** to the [sqlite3_bind_text | bindings] of that [parameter]. 
-** ^The specific value of WHERE-clause [parameter] might influence the 
+** ^The specific value of a WHERE-clause [parameter] might influence the 
 ** choice of query plan if the parameter is the left-hand side of a [LIKE]
 ** or [GLOB] operator or if the parameter is compared to an indexed column
-** and the [SQLITE_ENABLE_STAT3] compile-time option is enabled.
+** and the [SQLITE_ENABLE_STAT4] compile-time option is enabled.
 ** </li>
 ** </ol>
 **
@@ -4282,7 +4448,7 @@ SQLITE_API const void *sqlite3_column_name16(sqlite3_stmt*, int N);
 **
 ** ^If the Nth column returned by the statement is an expression or
 ** subquery and is not a column value, then all of these functions return
-** NULL.  ^These routine might also return NULL if a memory allocation error
+** NULL.  ^These routines might also return NULL if a memory allocation error
 ** occurs.  ^Otherwise, they return the name of the attached database, table,
 ** or column that query result column was extracted from.
 **
@@ -4292,10 +4458,6 @@ SQLITE_API const void *sqlite3_column_name16(sqlite3_stmt*, int N);
 ** ^These APIs are only available if the library was compiled with the
 ** [SQLITE_ENABLE_COLUMN_METADATA] C-preprocessor symbol.
 **
-** If two or more threads call one or more of these routines against the same
-** prepared statement and column at the same time then the results are
-** undefined.
-**
 ** If two or more threads call one or more
 ** [sqlite3_column_database_name | column metadata interfaces]
 ** for the same [prepared statement] and result column
@@ -4432,7 +4594,7 @@ SQLITE_API int sqlite3_step(sqlite3_stmt*);
 ** ^The sqlite3_data_count(P) interface returns the number of columns in the
 ** current row of the result set of [prepared statement] P.
 ** ^If prepared statement P does not have results ready to return
-** (via calls to the [sqlite3_column_int | sqlite3_column_*()] of
+** (via calls to the [sqlite3_column_int | sqlite3_column()] family of
 ** interfaces) then sqlite3_data_count(P) returns 0.
 ** ^The sqlite3_data_count(P) routine also returns 0 if P is a NULL pointer.
 ** ^The sqlite3_data_count(P) routine returns 0 if the previous call to
@@ -4756,8 +4918,6 @@ SQLITE_API int sqlite3_reset(sqlite3_stmt *pStmt);
 /*
 ** CAPI3REF: Create Or Redefine SQL Functions
 ** KEYWORDS: {function creation routines}
-** KEYWORDS: {application-defined SQL function}
-** KEYWORDS: {application-defined SQL functions}
 ** METHOD: sqlite3
 **
 ** ^These functions (collectively known as "function creation routines")
@@ -4811,6 +4971,23 @@ SQLITE_API int sqlite3_reset(sqlite3_stmt *pStmt);
 ** perform additional optimizations on deterministic functions, so use
 ** of the [SQLITE_DETERMINISTIC] flag is recommended where possible.
 **
+** ^The fourth parameter may also optionally include the [SQLITE_DIRECTONLY]
+** flag, which if present prevents the function from being invoked from
+** within VIEWs, TRIGGERs, CHECK constraints, generated column expressions,
+** index expressions, or the WHERE clause of partial indexes.
+**
+** <span style="background-color:#ffff90;">
+** For best security, the [SQLITE_DIRECTONLY] flag is recommended for
+** all application-defined SQL functions that do not need to be
+** used inside of triggers, view, CHECK constraints, or other elements of
+** the database schema.  This flags is especially recommended for SQL 
+** functions that have side effects or reveal internal application state.
+** Without this flag, an attacker might be able to modify the schema of
+** a database file to include invocations of the function with parameters
+** chosen by the attacker, which the application will then execute when
+** the database file is opened and read.
+** </span>
+**
 ** ^(The fifth parameter is an arbitrary pointer.  The implementation of the
 ** function can gain access to this pointer using [sqlite3_user_data()].)^
 **
@@ -4927,8 +5104,68 @@ SQLITE_API int sqlite3_create_window_function(
 ** [SQLITE_UTF8 | preferred text encoding] as the fourth argument
 ** to [sqlite3_create_function()], [sqlite3_create_function16()], or
 ** [sqlite3_create_function_v2()].
+**
+** <dl>
+** [[SQLITE_DETERMINISTIC]] <dt>SQLITE_DETERMINISTIC</dt><dd>
+** The SQLITE_DETERMINISTIC flag means that the new function always gives
+** the same output when the input parameters are the same.
+** The [abs|abs() function] is deterministic, for example, but
+** [randomblob|randomblob()] is not.  Functions must
+** be deterministic in order to be used in certain contexts such as
+** with the WHERE clause of [partial indexes] or in [generated columns].
+** SQLite might also optimize deterministic functions by factoring them
+** out of inner loops.
+** </dd>
+** 
+** [[SQLITE_DIRECTONLY]] <dt>SQLITE_DIRECTONLY</dt><dd>
+** The SQLITE_DIRECTONLY flag means that the function may only be invoked
+** from top-level SQL, and cannot be used in VIEWs or TRIGGERs nor in 
+** schema structures such as [CHECK constraints], [DEFAULT clauses],
+** [expression indexes], [partial indexes], or [generated columns].
+** The SQLITE_DIRECTONLY flags is a security feature which is recommended
+** for all [application-defined SQL functions], and especially for functions
+** that have side-effects or that could potentially leak sensitive
+** information.
+** </dd>
+**
+** [[SQLITE_INNOCUOUS]] <dt>SQLITE_INNOCUOUS</dt><dd>
+** The SQLITE_INNOCUOUS flag means that the function is unlikely
+** to cause problems even if misused.  An innocuous function should have
+** no side effects and should not depend on any values other than its
+** input parameters. The [abs|abs() function] is an example of an
+** innocuous function.
+** The [load_extension() SQL function] is not innocuous because of its
+** side effects.
+** <p> SQLITE_INNOCUOUS is similar to SQLITE_DETERMINISTIC, but is not
+** exactly the same.  The [random|random() function] is an example of a
+** function that is innocuous but not deterministic.
+** <p>Some heightened security settings
+** ([SQLITE_DBCONFIG_TRUSTED_SCHEMA] and [PRAGMA trusted_schema=OFF])
+** disable the use of SQL functions inside views and triggers and in
+** schema structures such as [CHECK constraints], [DEFAULT clauses],
+** [expression indexes], [partial indexes], and [generated columns] unless
+** the function is tagged with SQLITE_INNOCUOUS.  Most built-in functions
+** are innocuous.  Developers are advised to avoid using the
+** SQLITE_INNOCUOUS flag for application-defined functions unless the
+** function has been carefully audited and found to be free of potentially
+** security-adverse side-effects and information-leaks.
+** </dd>
+**
+** [[SQLITE_SUBTYPE]] <dt>SQLITE_SUBTYPE</dt><dd>
+** The SQLITE_SUBTYPE flag indicates to SQLite that a function may call
+** [sqlite3_value_subtype()] to inspect the sub-types of its arguments.
+** Specifying this flag makes no difference for scalar or aggregate user
+** functions. However, if it is not specified for a user-defined window
+** function, then any sub-types belonging to arguments passed to the window
+** function may be discarded before the window function is called (i.e.
+** sqlite3_value_subtype() will always return 0).
+** </dd>
+** </dl>
 */
-#define SQLITE_DETERMINISTIC    0x800
+#define SQLITE_DETERMINISTIC    0x000000800
+#define SQLITE_DIRECTONLY       0x000080000
+#define SQLITE_SUBTYPE          0x000100000
+#define SQLITE_INNOCUOUS        0x000200000
 
 /*
 ** CAPI3REF: Deprecated Functions
@@ -4987,8 +5224,8 @@ SQLITE_API SQLITE_DEPRECATED int sqlite3_memory_alarm(void(*)(void*,sqlite3_int6
 **
 ** These routines extract type, size, and content information from
 ** [protected sqlite3_value] objects.  Protected sqlite3_value objects
-** are used to pass parameter information into implementation of
-** [application-defined SQL functions] and [virtual tables].
+** are used to pass parameter information into the functions that
+** implement [application-defined SQL functions] and [virtual tables].
 **
 ** These routines work only with [protected sqlite3_value] objects.
 ** Any attempt to use these routines on an [unprotected sqlite3_value]
@@ -5045,7 +5282,7 @@ SQLITE_API SQLITE_DEPRECATED int sqlite3_memory_alarm(void(*)(void*,sqlite3_int6
 ** ^The sqlite3_value_frombind(X) interface returns non-zero if the
 ** value X originated from one of the [sqlite3_bind_int|sqlite3_bind()]
 ** interfaces.  ^If X comes from an SQL literal value, or a table column,
-** and expression, then sqlite3_value_frombind(X) returns zero.
+** or an expression, then sqlite3_value_frombind(X) returns zero.
 **
 ** Please pay particular attention to the fact that the pointer returned
 ** from [sqlite3_value_blob()], [sqlite3_value_text()], or
@@ -5131,8 +5368,8 @@ SQLITE_API void sqlite3_value_free(sqlite3_value*);
 ** routine to allocate memory for storing their state.
 **
 ** ^The first time the sqlite3_aggregate_context(C,N) routine is called 
-** for a particular aggregate function, SQLite
-** allocates N of memory, zeroes out that memory, and returns a pointer
+** for a particular aggregate function, SQLite allocates
+** N bytes of memory, zeroes out that memory, and returns a pointer
 ** to the new memory. ^On second and subsequent calls to
 ** sqlite3_aggregate_context() for the same aggregate function instance,
 ** the same buffer is returned.  Sqlite3_aggregate_context() is normally
@@ -5149,7 +5386,7 @@ SQLITE_API void sqlite3_value_free(sqlite3_value*);
 **
 ** ^(The amount of space allocated by sqlite3_aggregate_context(C,N) is
 ** determined by the N parameter on first successful call.  Changing the
-** value of N in subsequent call to sqlite3_aggregate_context() within
+** value of N in any subsequents call to sqlite3_aggregate_context() within
 ** the same aggregate function instance will not resize the memory
 ** allocation.)^  Within the xFinal callback, it is customary to set
 ** N=0 in calls to sqlite3_aggregate_context(C,N) so that no 
@@ -5460,7 +5697,7 @@ SQLITE_API void sqlite3_result_subtype(sqlite3_context*,unsigned int);
 ** <li> [SQLITE_UTF16_ALIGNED].
 ** </ul>)^
 ** ^The eTextRep argument determines the encoding of strings passed
-** to the collating function callback, xCallback.
+** to the collating function callback, xCompare.
 ** ^The [SQLITE_UTF16] and [SQLITE_UTF16_ALIGNED] values for eTextRep
 ** force strings to be UTF16 with native byte order.
 ** ^The [SQLITE_UTF16_ALIGNED] value for eTextRep forces strings to begin
@@ -5469,18 +5706,19 @@ SQLITE_API void sqlite3_result_subtype(sqlite3_context*,unsigned int);
 ** ^The fourth argument, pArg, is an application data pointer that is passed
 ** through as the first argument to the collating function callback.
 **
-** ^The fifth argument, xCallback, is a pointer to the collating function.
+** ^The fifth argument, xCompare, is a pointer to the collating function.
 ** ^Multiple collating functions can be registered using the same name but
 ** with different eTextRep parameters and SQLite will use whichever
 ** function requires the least amount of data transformation.
-** ^If the xCallback argument is NULL then the collating function is
+** ^If the xCompare argument is NULL then the collating function is
 ** deleted.  ^When all collating functions having the same name are deleted,
 ** that collation is no longer usable.
 **
 ** ^The collating function callback is invoked with a copy of the pArg 
 ** application data pointer and with two strings in the encoding specified
-** by the eTextRep argument.  The collating function must return an
-** integer that is negative, zero, or positive
+** by the eTextRep argument.  The two integer parameters to the collating
+** function callback are the length of the two strings, in bytes. The collating
+** function must return an integer that is negative, zero, or positive
 ** if the first string is less than, equal to, or greater than the second,
 ** respectively.  A collating function must always return the same answer
 ** given the same inputs.  If two or more collating functions are registered
@@ -5497,7 +5735,7 @@ SQLITE_API void sqlite3_result_subtype(sqlite3_context*,unsigned int);
 ** </ol>
 **
 ** If a collating function fails any of the above constraints and that
-** collating function is  registered and used, then the behavior of SQLite
+** collating function is registered and used, then the behavior of SQLite
 ** is undefined.
 **
 ** ^The sqlite3_create_collation_v2() works like sqlite3_create_collation()
@@ -5824,16 +6062,31 @@ SQLITE_API sqlite3 *sqlite3_db_handle(sqlite3_stmt*);
 ** CAPI3REF: Return The Filename For A Database Connection
 ** METHOD: sqlite3
 **
-** ^The sqlite3_db_filename(D,N) interface returns a pointer to a filename
-** associated with database N of connection D.  ^The main database file
-** has the name "main".  If there is no attached database N on the database
+** ^The sqlite3_db_filename(D,N) interface returns a pointer to the filename
+** associated with database N of connection D.
+** ^If there is no attached database N on the database
 ** connection D, or if database N is a temporary or in-memory database, then
 ** this function will return either a NULL pointer or an empty string.
 **
+** ^The string value returned by this routine is owned and managed by
+** the database connection.  ^The value will be valid until the database N
+** is [DETACH]-ed or until the database connection closes.
+**
 ** ^The filename returned by this function is the output of the
 ** xFullPathname method of the [VFS].  ^In other words, the filename
 ** will be an absolute pathname, even if the filename used
 ** to open the database originally was a URI or relative pathname.
+**
+** If the filename pointer returned by this routine is not NULL, then it
+** can be used as the filename input parameter to these routines:
+** <ul>
+** <li> [sqlite3_uri_parameter()]
+** <li> [sqlite3_uri_boolean()]
+** <li> [sqlite3_uri_int64()]
+** <li> [sqlite3_filename_database()]
+** <li> [sqlite3_filename_journal()]
+** <li> [sqlite3_filename_wal()]
+** </ul>
 */
 SQLITE_API const char *sqlite3_db_filename(sqlite3 *db, const char *zDbName);
 
@@ -5983,15 +6236,19 @@ SQLITE_API void *sqlite3_update_hook(
 **
 ** ^(The cache sharing mode set by this interface effects all subsequent
 ** calls to [sqlite3_open()], [sqlite3_open_v2()], and [sqlite3_open16()].
-** Existing database connections continue use the sharing mode
+** Existing database connections continue to use the sharing mode
 ** that was in effect at the time they were opened.)^
 **
 ** ^(This routine returns [SQLITE_OK] if shared cache was enabled or disabled
 ** successfully.  An [error code] is returned otherwise.)^
 **
-** ^Shared cache is disabled by default. But this might change in
-** future releases of SQLite.  Applications that care about shared
-** cache setting should set it explicitly.
+** ^Shared cache is disabled by default. It is recommended that it stay
+** that way.  In other words, do not use this routine.  This interface
+** continues to be provided for historical compatibility, but its use is
+** discouraged.  Any use of shared cache is discouraged.  If shared cache
+** must be used, it is recommended that shared cache only be enabled for
+** individual database connections using the [sqlite3_open_v2()] interface
+** with the [SQLITE_OPEN_SHAREDCACHE] flag.
 **
 ** Note: This method is disabled on MacOS X 10.7 and iOS version 5.0
 ** and will always return SQLITE_MISUSE. On those systems, 
@@ -6038,6 +6295,9 @@ SQLITE_API int sqlite3_db_release_memory(sqlite3*);
 /*
 ** CAPI3REF: Impose A Limit On Heap Size
 **
+** These interfaces impose limits on the amount of heap memory that will be
+** by all database connections within a single process.
+**
 ** ^The sqlite3_soft_heap_limit64() interface sets and/or queries the
 ** soft limit on the amount of heap memory that may be allocated by SQLite.
 ** ^SQLite strives to keep heap memory utilization below the soft heap
@@ -6048,20 +6308,41 @@ SQLITE_API int sqlite3_db_release_memory(sqlite3*);
 ** an [SQLITE_NOMEM] error.  In other words, the soft heap limit 
 ** is advisory only.
 **
-** ^The return value from sqlite3_soft_heap_limit64() is the size of
-** the soft heap limit prior to the call, or negative in the case of an
-** error.  ^If the argument N is negative
-** then no change is made to the soft heap limit.  Hence, the current
-** size of the soft heap limit can be determined by invoking
-** sqlite3_soft_heap_limit64() with a negative argument.
-**
-** ^If the argument N is zero then the soft heap limit is disabled.
+** ^The sqlite3_hard_heap_limit64(N) interface sets a hard upper bound of
+** N bytes on the amount of memory that will be allocated.  ^The
+** sqlite3_hard_heap_limit64(N) interface is similar to
+** sqlite3_soft_heap_limit64(N) except that memory allocations will fail
+** when the hard heap limit is reached.
 **
-** ^(The soft heap limit is not enforced in the current implementation
+** ^The return value from both sqlite3_soft_heap_limit64() and
+** sqlite3_hard_heap_limit64() is the size of
+** the heap limit prior to the call, or negative in the case of an
+** error.  ^If the argument N is negative
+** then no change is made to the heap limit.  Hence, the current
+** size of heap limits can be determined by invoking
+** sqlite3_soft_heap_limit64(-1) or sqlite3_hard_heap_limit(-1).
+**
+** ^Setting the heap limits to zero disables the heap limiter mechanism.
+**
+** ^The soft heap limit may not be greater than the hard heap limit.
+** ^If the hard heap limit is enabled and if sqlite3_soft_heap_limit(N)
+** is invoked with a value of N that is greater than the hard heap limit,
+** the the soft heap limit is set to the value of the hard heap limit.
+** ^The soft heap limit is automatically enabled whenever the hard heap
+** limit is enabled. ^When sqlite3_hard_heap_limit64(N) is invoked and
+** the soft heap limit is outside the range of 1..N, then the soft heap
+** limit is set to N.  ^Invoking sqlite3_soft_heap_limit64(0) when the
+** hard heap limit is enabled makes the soft heap limit equal to the
+** hard heap limit.
+**
+** The memory allocation limits can also be adjusted using
+** [PRAGMA soft_heap_limit] and [PRAGMA hard_heap_limit].
+**
+** ^(The heap limits are not enforced in the current implementation
 ** if one or more of following conditions are true:
 **
 ** <ul>
-** <li> The soft heap limit is set to zero.
+** <li> The limit value is set to zero.
 ** <li> Memory accounting is disabled using a combination of the
 **      [sqlite3_config]([SQLITE_CONFIG_MEMSTATUS],...) start-time option and
 **      the [SQLITE_DEFAULT_MEMSTATUS] compile-time option.
@@ -6072,21 +6353,11 @@ SQLITE_API int sqlite3_db_release_memory(sqlite3*);
 **      from the heap.
 ** </ul>)^
 **
-** Beginning with SQLite [version 3.7.3] ([dateof:3.7.3]), 
-** the soft heap limit is enforced
-** regardless of whether or not the [SQLITE_ENABLE_MEMORY_MANAGEMENT]
-** compile-time option is invoked.  With [SQLITE_ENABLE_MEMORY_MANAGEMENT],
-** the soft heap limit is enforced on every memory allocation.  Without
-** [SQLITE_ENABLE_MEMORY_MANAGEMENT], the soft heap limit is only enforced
-** when memory is allocated by the page cache.  Testing suggests that because
-** the page cache is the predominate memory user in SQLite, most
-** applications will achieve adequate soft heap limit enforcement without
-** the use of [SQLITE_ENABLE_MEMORY_MANAGEMENT].
-**
-** The circumstances under which SQLite will enforce the soft heap limit may
+** The circumstances under which SQLite will enforce the heap limits may
 ** changes in future releases of SQLite.
 */
 SQLITE_API sqlite3_int64 sqlite3_soft_heap_limit64(sqlite3_int64 N);
+SQLITE_API sqlite3_int64 sqlite3_hard_heap_limit64(sqlite3_int64 N);
 
 /*
 ** CAPI3REF: Deprecated Soft Heap Limit Interface
@@ -6110,7 +6381,7 @@ SQLITE_API SQLITE_DEPRECATED void sqlite3_soft_heap_limit(int N);
 ** interface returns SQLITE_OK and fills in the non-NULL pointers in
 ** the final five arguments with appropriate values if the specified
 ** column exists.  ^The sqlite3_table_column_metadata() interface returns
-** SQLITE_ERROR and if the specified column does not exist.
+** SQLITE_ERROR if the specified column does not exist.
 ** ^If the column-name parameter to sqlite3_table_column_metadata() is a
 ** NULL pointer, then this routine simply checks for the existence of the
 ** table and returns SQLITE_OK if the table exists and SQLITE_ERROR if it
@@ -6252,7 +6523,7 @@ SQLITE_API int sqlite3_load_extension(
 ** to enable or disable only the C-API.)^
 **
 ** <b>Security warning:</b> It is recommended that extension loading
-** be disabled using the [SQLITE_DBCONFIG_ENABLE_LOAD_EXTENSION] method
+** be enabled using the [SQLITE_DBCONFIG_ENABLE_LOAD_EXTENSION] method
 ** rather than this interface, so the [load_extension()] SQL function
 ** remains disabled. This will prevent SQL injections from giving attackers
 ** access to extension loading capabilities.
@@ -6339,7 +6610,7 @@ typedef struct sqlite3_module sqlite3_module;
 ** KEYWORDS: sqlite3_module {virtual table module}
 **
 ** This structure, sometimes called a "virtual table module", 
-** defines the implementation of a [virtual tables].  
+** defines the implementation of a [virtual table].  
 ** This structure consists mostly of methods for the module.
 **
 ** ^A virtual table module is created by filling in a persistent
@@ -6436,7 +6707,13 @@ struct sqlite3_module {
 ** the right-hand side of the corresponding aConstraint[] is evaluated
 ** and becomes the argvIndex-th entry in argv.  ^(If aConstraintUsage[].omit
 ** is true, then the constraint is assumed to be fully handled by the
-** virtual table and is not checked again by SQLite.)^
+** virtual table and might not be checked again by the byte code.)^ ^(The
+** aConstraintUsage[].omit flag is an optimization hint. When the omit flag
+** is left in its default setting of false, the constraint will always be
+** checked separately in byte code.  If the omit flag is change to true, then
+** the constraint may or may not be checked in byte code.  In other words,
+** when the omit flag is true there is no guarantee that the constraint will
+** not be checked again using byte code.)^
 **
 ** ^The idxNum and idxPtr values are recorded and passed into the
 ** [xFilter] method.
@@ -6476,7 +6753,7 @@ struct sqlite3_module {
 ** If a virtual table extension is
 ** used with an SQLite version earlier than 3.8.2, the results of attempting 
 ** to read or write the estimatedRows field are undefined (but are likely 
-** to included crashing the application). The estimatedRows field should
+** to include crashing the application). The estimatedRows field should
 ** therefore only be used if [sqlite3_libversion_number()] returns a
 ** value greater than or equal to 3008002. Similarly, the idxFlags field
 ** was added for [version 3.9.0] ([dateof:3.9.0]). 
@@ -6528,7 +6805,7 @@ struct sqlite3_index_info {
 /*
 ** CAPI3REF: Virtual Table Constraint Operator Codes
 **
-** These macros defined the allowed values for the
+** These macros define the allowed values for the
 ** [sqlite3_index_info].aConstraint[].op field.  Each value represents
 ** an operator that is part of a constraint term in the wHERE clause of
 ** a query that uses a [virtual table].
@@ -6574,6 +6851,12 @@ struct sqlite3_index_info {
 ** ^The sqlite3_create_module()
 ** interface is equivalent to sqlite3_create_module_v2() with a NULL
 ** destructor.
+**
+** ^If the third parameter (the pointer to the sqlite3_module object) is
+** NULL then no new module is create and any existing modules with the
+** same name are dropped.
+**
+** See also: [sqlite3_drop_modules()]
 */
 SQLITE_API int sqlite3_create_module(
   sqlite3 *db,               /* SQLite connection to register module with */
@@ -6589,6 +6872,23 @@ SQLITE_API int sqlite3_create_module_v2(
   void(*xDestroy)(void*)     /* Module destructor function */
 );
 
+/*
+** CAPI3REF: Remove Unnecessary Virtual Table Implementations
+** METHOD: sqlite3
+**
+** ^The sqlite3_drop_modules(D,L) interface removes all virtual
+** table modules from database connection D except those named on list L.
+** The L parameter must be either NULL or a pointer to an array of pointers
+** to strings where the array is terminated by a single NULL pointer.
+** ^If the L parameter is NULL, then all virtual table modules are removed.
+**
+** See also: [sqlite3_create_module()]
+*/
+SQLITE_API int sqlite3_drop_modules(
+  sqlite3 *db,                /* Remove modules from this connection */
+  const char **azKeep         /* Except, do not remove the ones named here */
+);
+
 /*
 ** CAPI3REF: Virtual Table Instance Object
 ** KEYWORDS: sqlite3_vtab
@@ -7115,7 +7415,7 @@ SQLITE_API void sqlite3_mutex_leave(sqlite3_mutex*);
 ** The only difference is that the public sqlite3_XXX functions enumerated
 ** above silently ignore any invocations that pass a NULL pointer instead
 ** of a valid mutex handle. The implementations of the methods defined
-** by this structure are not required to handle this case, the results
+** by this structure are not required to handle this case. The results
 ** of passing a NULL pointer instead of a valid mutex handle are undefined
 ** (i.e. it is acceptable to provide an implementation that segfaults if
 ** it is passed a NULL pointer).
@@ -7297,7 +7597,7 @@ SQLITE_API int sqlite3_test_control(int op, ...);
 #define SQLITE_TESTCTRL_FIRST                    5
 #define SQLITE_TESTCTRL_PRNG_SAVE                5
 #define SQLITE_TESTCTRL_PRNG_RESTORE             6
-#define SQLITE_TESTCTRL_PRNG_RESET               7
+#define SQLITE_TESTCTRL_PRNG_RESET               7  /* NOT USED */
 #define SQLITE_TESTCTRL_BITVEC_TEST              8
 #define SQLITE_TESTCTRL_FAULT_INSTALL            9
 #define SQLITE_TESTCTRL_BENIGN_MALLOC_HOOKS     10
@@ -7319,7 +7619,10 @@ SQLITE_API int sqlite3_test_control(int op, ...);
 #define SQLITE_TESTCTRL_SORTER_MMAP             24
 #define SQLITE_TESTCTRL_IMPOSTER                25
 #define SQLITE_TESTCTRL_PARSER_COVERAGE         26
-#define SQLITE_TESTCTRL_LAST                    26  /* Largest TESTCTRL */
+#define SQLITE_TESTCTRL_RESULT_INTREAL          27
+#define SQLITE_TESTCTRL_PRNG_SEED               28
+#define SQLITE_TESTCTRL_EXTRA_SCHEMA_CHECKS     29
+#define SQLITE_TESTCTRL_LAST                    29  /* Largest TESTCTRL */
 
 /*
 ** CAPI3REF: SQL Keyword Checking
@@ -7585,7 +7888,7 @@ SQLITE_API int sqlite3_status64(
 **
 ** [[SQLITE_STATUS_PAGECACHE_SIZE]] ^(<dt>SQLITE_STATUS_PAGECACHE_SIZE</dt>
 ** <dd>This parameter records the largest memory allocation request
-** handed to [pagecache memory allocator].  Only the value returned in the
+** handed to the [pagecache memory allocator].  Only the value returned in the
 ** *pHighwater parameter to [sqlite3_status()] is of interest.  
 ** The value written into the *pCurrent parameter is undefined.</dd>)^
 **
@@ -7661,7 +7964,7 @@ SQLITE_API int sqlite3_db_status(sqlite3*, int op, int *pCur, int *pHiwtr, int r
 ** checked out.</dd>)^
 **
 ** [[SQLITE_DBSTATUS_LOOKASIDE_HIT]] ^(<dt>SQLITE_DBSTATUS_LOOKASIDE_HIT</dt>
-** <dd>This parameter returns the number malloc attempts that were 
+** <dd>This parameter returns the number of malloc attempts that were 
 ** satisfied using lookaside memory. Only the high-water value is meaningful;
 ** the current value is always zero.)^
 **
@@ -7743,7 +8046,7 @@ SQLITE_API int sqlite3_db_status(sqlite3*, int op, int *pCur, int *pHiwtr, int r
 ** cache overflowing. Transactions are more efficient if they are written
 ** to disk all at once. When pages spill mid-transaction, that introduces
 ** additional overhead. This parameter can be used help identify
-** inefficiencies that can be resolve by increasing the cache size.
+** inefficiencies that can be resolved by increasing the cache size.
 ** </dd>
 **
 ** [[SQLITE_DBSTATUS_DEFERRED_FKS]] ^(<dt>SQLITE_DBSTATUS_DEFERRED_FKS</dt>
@@ -7832,7 +8135,7 @@ SQLITE_API int sqlite3_stmt_status(sqlite3_stmt*, int op,int resetFlg);
 **
 ** [[SQLITE_STMTSTATUS_REPREPARE]] <dt>SQLITE_STMTSTATUS_REPREPARE</dt>
 ** <dd>^This is the number of times that the prepare statement has been
-** automatically regenerated due to schema changes or change to 
+** automatically regenerated due to schema changes or changes to 
 ** [bound parameters] that might affect the query plan.
 **
 ** [[SQLITE_STMTSTATUS_RUN]] <dt>SQLITE_STMTSTATUS_RUN</dt>
@@ -8003,7 +8306,7 @@ struct sqlite3_pcache_page {
 **
 ** ^(SQLite will normally invoke xFetch() with a createFlag of 0 or 1.  SQLite
 ** will only use a createFlag of 2 after a prior call with a createFlag of 1
-** failed.)^  In between the to xFetch() calls, SQLite may
+** failed.)^  In between the xFetch() calls, SQLite may
 ** attempt to unpin one or more cache pages by spilling the content of
 ** pinned pages to disk and synching the operating system disk cache.
 **
@@ -8321,7 +8624,7 @@ SQLITE_API int sqlite3_backup_pagecount(sqlite3_backup *p);
 ** the first argument to register for a callback that will be invoked
 ** when the blocking connections current transaction is concluded. ^The
 ** callback is invoked from within the [sqlite3_step] or [sqlite3_close]
-** call that concludes the blocking connections transaction.
+** call that concludes the blocking connection's transaction.
 **
 ** ^(If sqlite3_unlock_notify() is called in a multi-threaded application,
 ** there is a chance that the blocking connection will have already
@@ -8359,7 +8662,7 @@ SQLITE_API int sqlite3_backup_pagecount(sqlite3_backup *p);
 ** an unlock-notify callback is a pointer to an array of void* pointers,
 ** and the second is the number of entries in the array.
 **
-** When a blocking connections transaction is concluded, there may be
+** When a blocking connection's transaction is concluded, there may be
 ** more than one blocked connection that has registered for an unlock-notify
 ** callback. ^If two or more such blocked connections have specified the
 ** same callback function, then instead of invoking the callback function
@@ -8707,14 +9010,20 @@ SQLITE_API int sqlite3_wal_checkpoint_v2(
 ** If this interface is invoked outside the context of an xConnect or
 ** xCreate virtual table method then the behavior is undefined.
 **
-** At present, there is only one option that may be configured using
-** this function. (See [SQLITE_VTAB_CONSTRAINT_SUPPORT].)  Further options
-** may be added in the future.
+** In the call sqlite3_vtab_config(D,C,...) the D parameter is the
+** [database connection] in which the virtual table is being created and
+** which is passed in as the first argument to the [xConnect] or [xCreate]
+** method that is invoking sqlite3_vtab_config().  The C parameter is one
+** of the [virtual table configuration options].  The presence and meaning
+** of parameters after C depend on which [virtual table configuration option]
+** is used.
 */
 SQLITE_API int sqlite3_vtab_config(sqlite3*, int op, ...);
 
 /*
 ** CAPI3REF: Virtual Table Configuration Options
+** KEYWORDS: {virtual table configuration options} 
+** KEYWORDS: {virtual table configuration option}
 **
 ** These macros define the various options to the
 ** [sqlite3_vtab_config()] interface that [virtual table] implementations
@@ -8722,7 +9031,7 @@ SQLITE_API int sqlite3_vtab_config(sqlite3*, int op, ...);
 **
 ** <dl>
 ** [[SQLITE_VTAB_CONSTRAINT_SUPPORT]]
-** <dt>SQLITE_VTAB_CONSTRAINT_SUPPORT
+** <dt>SQLITE_VTAB_CONSTRAINT_SUPPORT</dt>
 ** <dd>Calls of the form
 ** [sqlite3_vtab_config](db,SQLITE_VTAB_CONSTRAINT_SUPPORT,X) are supported,
 ** where X is an integer.  If X is zero, then the [virtual table] whose
@@ -8751,9 +9060,31 @@ SQLITE_API int sqlite3_vtab_config(sqlite3*, int op, ...);
 ** return SQLITE_OK. Or, if this is not possible, it may return
 ** SQLITE_CONSTRAINT, in which case SQLite falls back to OR ABORT 
 ** constraint handling.
+** </dd>
+**
+** [[SQLITE_VTAB_DIRECTONLY]]<dt>SQLITE_VTAB_DIRECTONLY</dt>
+** <dd>Calls of the form
+** [sqlite3_vtab_config](db,SQLITE_VTAB_DIRECTONLY) from within the
+** the [xConnect] or [xCreate] methods of a [virtual table] implmentation
+** prohibits that virtual table from being used from within triggers and
+** views.
+** </dd>
+**
+** [[SQLITE_VTAB_INNOCUOUS]]<dt>SQLITE_VTAB_INNOCUOUS</dt>
+** <dd>Calls of the form
+** [sqlite3_vtab_config](db,SQLITE_VTAB_INNOCUOUS) from within the
+** the [xConnect] or [xCreate] methods of a [virtual table] implmentation
+** identify that virtual table as being safe to use from within triggers
+** and views.  Conceptually, the SQLITE_VTAB_INNOCUOUS tag means that the
+** virtual table can do no serious harm even if it is controlled by a
+** malicious hacker.  Developers should avoid setting the SQLITE_VTAB_INNOCUOUS
+** flag unless absolutely necessary.
+** </dd>
 ** </dl>
 */
 #define SQLITE_VTAB_CONSTRAINT_SUPPORT 1
+#define SQLITE_VTAB_INNOCUOUS          2
+#define SQLITE_VTAB_DIRECTONLY         3
 
 /*
 ** CAPI3REF: Determine The Virtual Table Conflict Policy
@@ -8833,15 +9164,15 @@ SQLITE_API SQLITE_EXPERIMENTAL const char *sqlite3_vtab_collation(sqlite3_index_
 **
 ** <dl>
 ** [[SQLITE_SCANSTAT_NLOOP]] <dt>SQLITE_SCANSTAT_NLOOP</dt>
-** <dd>^The [sqlite3_int64] variable pointed to by the T parameter will be
+** <dd>^The [sqlite3_int64] variable pointed to by the V parameter will be
 ** set to the total number of times that the X-th loop has run.</dd>
 **
 ** [[SQLITE_SCANSTAT_NVISIT]] <dt>SQLITE_SCANSTAT_NVISIT</dt>
-** <dd>^The [sqlite3_int64] variable pointed to by the T parameter will be set
+** <dd>^The [sqlite3_int64] variable pointed to by the V parameter will be set
 ** to the total number of rows examined by all iterations of the X-th loop.</dd>
 **
 ** [[SQLITE_SCANSTAT_EST]] <dt>SQLITE_SCANSTAT_EST</dt>
-** <dd>^The "double" variable pointed to by the T parameter will be set to the
+** <dd>^The "double" variable pointed to by the V parameter will be set to the
 ** query planner's estimate for the average number of rows output from each
 ** iteration of the X-th loop.  If the query planner's estimates was accurate,
 ** then this value will approximate the quotient NVISIT/NLOOP and the
@@ -8849,17 +9180,17 @@ SQLITE_API SQLITE_EXPERIMENTAL const char *sqlite3_vtab_collation(sqlite3_index_
 ** be the NLOOP value for the current loop.
 **
 ** [[SQLITE_SCANSTAT_NAME]] <dt>SQLITE_SCANSTAT_NAME</dt>
-** <dd>^The "const char *" variable pointed to by the T parameter will be set
+** <dd>^The "const char *" variable pointed to by the V parameter will be set
 ** to a zero-terminated UTF-8 string containing the name of the index or table
 ** used for the X-th loop.
 **
 ** [[SQLITE_SCANSTAT_EXPLAIN]] <dt>SQLITE_SCANSTAT_EXPLAIN</dt>
-** <dd>^The "const char *" variable pointed to by the T parameter will be set
+** <dd>^The "const char *" variable pointed to by the V parameter will be set
 ** to a zero-terminated UTF-8 string containing the [EXPLAIN QUERY PLAN]
 ** description for the X-th loop.
 **
 ** [[SQLITE_SCANSTAT_SELECTID]] <dt>SQLITE_SCANSTAT_SELECT</dt>
-** <dd>^The "int" variable pointed to by the T parameter will be set to the
+** <dd>^The "int" variable pointed to by the V parameter will be set to the
 ** "select-id" for the X-th loop.  The select-id identifies which query or
 ** subquery the loop is part of.  The main query has a select-id of zero.
 ** The select-id is the same value as is output in the first column
@@ -9714,7 +10045,7 @@ SQLITE_API int sqlite3session_attach(
 ** The second argument (xFilter) is the "filter callback". For changes to rows 
 ** in tables that are not attached to the Session object, the filter is called
 ** to determine whether changes to the table's rows should be tracked or not. 
-** If xFilter returns 0, changes is not tracked. Note that once a table is 
+** If xFilter returns 0, changes are not tracked. Note that once a table is 
 ** attached, xFilter will not be called again.
 */
 SQLITE_API void sqlite3session_table_filter(
@@ -9888,7 +10219,7 @@ SQLITE_API int sqlite3session_changeset(
 ** It an error if database zFrom does not exist or does not contain the
 ** required compatible table.
 **
-** If the operation successful, SQLITE_OK is returned. Otherwise, an SQLite
+** If the operation is successful, SQLITE_OK is returned. Otherwise, an SQLite
 ** error code. In this case, if argument pzErrMsg is not NULL, *pzErrMsg
 ** may be set to point to a buffer containing an English language error 
 ** message. It is the responsibility of the caller to free this buffer using
@@ -10025,7 +10356,7 @@ SQLITE_API int sqlite3changeset_start_v2(
 ** CAPI3REF: Advance A Changeset Iterator
 ** METHOD: sqlite3_changeset_iter
 **
-** This function may only be used with iterators created by function
+** This function may only be used with iterators created by the function
 ** [sqlite3changeset_start()]. If it is called on an iterator passed to
 ** a conflict-handler callback by [sqlite3changeset_apply()], SQLITE_MISUSE
 ** is returned and the call has no effect.
@@ -10441,8 +10772,8 @@ SQLITE_API int sqlite3changegroup_new(sqlite3_changegroup **pp);
 ** case, this function fails with SQLITE_SCHEMA. If the input changeset
 ** appears to be corrupt and the corruption is detected, SQLITE_CORRUPT is
 ** returned. Or, if an out-of-memory condition occurs during processing, this
-** function returns SQLITE_NOMEM. In all cases, if an error occurs the
-** final contents of the changegroup is undefined.
+** function returns SQLITE_NOMEM. In all cases, if an error occurs the state
+** of the final contents of the changegroup is undefined.
 **
 ** If no error occurs, SQLITE_OK is returned.
 */
@@ -10617,7 +10948,7 @@ SQLITE_API void sqlite3changegroup_delete(sqlite3_changegroup*);
 **
 ** It is safe to execute SQL statements, including those that write to the
 ** table that the callback related to, from within the xConflict callback.
-** This can be used to further customize the applications conflict
+** This can be used to further customize the application's conflict
 ** resolution strategy.
 **
 ** All changes made by these functions are enclosed in a savepoint transaction.
@@ -10927,7 +11258,7 @@ SQLITE_API int sqlite3rebaser_configure(
 **
 ** Argument pIn must point to a buffer containing a changeset nIn bytes
 ** in size. This function allocates and populates a buffer with a copy
-** of the changeset rebased rebased according to the configuration of the
+** of the changeset rebased according to the configuration of the
 ** rebaser object passed as the first argument. If successful, (*ppOut)
 ** is set to point to the new buffer containing the rebased changeset and 
 ** (*pnOut) to its size in bytes and SQLITE_OK returned. It is the
@@ -11335,7 +11666,7 @@ struct Fts5PhraseIter {
 **
 ** xSetAuxdata(pFts5, pAux, xDelete)
 **
-**   Save the pointer passed as the second argument as the extension functions 
+**   Save the pointer passed as the second argument as the extension function's 
 **   "auxiliary data". The pointer may then be retrieved by the current or any
 **   future invocation of the same fts5 extension function made as part of
 **   the same MATCH query using the xGetAuxdata() API.
@@ -11577,8 +11908,8 @@ struct Fts5ExtensionApi {
 **
 **   There are several ways to approach this in FTS5:
 **
-**   <ol><li> By mapping all synonyms to a single token. In this case, the 
-**            In the above example, this means that the tokenizer returns the
+**   <ol><li> By mapping all synonyms to a single token. In this case, using
+**            the above example, this means that the tokenizer returns the
 **            same token for inputs "first" and "1st". Say that token is in
 **            fact "first", so that when the user inserts the document "I won
 **            1st place" entries are added to the index for tokens "i", "won",
