Updating SQLite pull to newest release.  Enabling additional #defines to allow for JSON support.

diff --git a/Engine/Plugins/Runtime/Database/SQLiteCore/Source/SQLiteCore/Private/sqlite/sqlite3.inl b/Engine/Plugins/Runtime/Database/SQLiteCore/Source/SQLiteCore/Private/sqlite/sqlite3.inl
--- a/Engine/Plugins/Runtime/Database/SQLiteCore/Source/SQLiteCore/Private/sqlite/sqlite3.inl
+++ b/Engine/Plugins/Runtime/Database/SQLiteCore/Source/SQLiteCore/Private/sqlite/sqlite3.inl
@@ -1,6 +1,6 @@
 /******************************************************************************
 ** This file is an amalgamation of many separate C source files from SQLite
-** version 3.28.0.  By combining all the individual C code files into this
+** version 3.31.1.  By combining all the individual C code files into this
 ** single large file, the entire code can be compiled as a single translation
 ** unit.  This allows many compilers to do optimizations that would not be
 ** possible if the files were compiled separately.  Performance improvements
@@ -39,7 +39,7 @@
 ** SQLite was built with.
 */
 
-#ifndef SQLITE_OMIT_COMPILEOPTION_DIAGS
+#ifndef SQLITE_OMIT_COMPILEOPTION_DIAGS /* IMP: R-16824-07538 */
 
 /*
 ** Include the configuration header output by 'configure' if we're using the
@@ -331,8 +331,6 @@ static const char * const sqlite3azCompileOpt[] = {
 #endif
 #if defined(SQLITE_ENABLE_STAT4)
   "ENABLE_STAT4",
-#elif defined(SQLITE_ENABLE_STAT3)
-  "ENABLE_STAT3",
 #endif
 #if SQLITE_ENABLE_STMTVTAB
   "ENABLE_STMTVTAB",
@@ -888,6 +886,11 @@ SQLITE_PRIVATE const char **sqlite3CompileOptions(int *pnOpt){
 #pragma warning(disable : 4706)
 #endif /* defined(_MSC_VER) */
 
+#if defined(_MSC_VER) && !defined(_WIN64)
+#undef SQLITE_4_BYTE_ALIGNED_MALLOC
+#define SQLITE_4_BYTE_ALIGNED_MALLOC
+#endif /* defined(_MSC_VER) && !defined(_WIN64) */
+
 #endif /* SQLITE_MSVC_H */
 
 /************** End of msvc.h ************************************************/
@@ -1162,9 +1165,9 @@ extern "C" {
 ** [sqlite3_libversion_number()], [sqlite3_sourceid()],
 ** [sqlite_version()] and [sqlite_source_id()].
 */
-#define SQLITE_VERSION        "3.28.0"
-#define SQLITE_VERSION_NUMBER 3028000
-#define SQLITE_SOURCE_ID      "2019-04-16 19:49:53 884b4b7e502b4e991677b53971277adfaf0a04a284f8e483e2553d0f83156b50"
+#define SQLITE_VERSION        "3.31.1"
+#define SQLITE_VERSION_NUMBER 3031001
+#define SQLITE_SOURCE_ID      "2020-01-27 19:55:54 3bfa9cc97da10598521b342961df8f5f68c7388fa117345eeb516eaa837bb4d6"
 
 /*
 ** CAPI3REF: Run-Time Library Version Numbers
@@ -1555,6 +1558,7 @@ SQLITE_API int sqlite3_exec(
 #define SQLITE_CANTOPEN_FULLPATH       (SQLITE_CANTOPEN | (3<<8))
 #define SQLITE_CANTOPEN_CONVPATH       (SQLITE_CANTOPEN | (4<<8))
 #define SQLITE_CANTOPEN_DIRTYWAL       (SQLITE_CANTOPEN | (5<<8)) /* Not Used */
+#define SQLITE_CANTOPEN_SYMLINK        (SQLITE_CANTOPEN | (6<<8))
 #define SQLITE_CORRUPT_VTAB            (SQLITE_CORRUPT | (1<<8))
 #define SQLITE_CORRUPT_SEQUENCE        (SQLITE_CORRUPT | (2<<8))
 #define SQLITE_READONLY_RECOVERY       (SQLITE_READONLY | (1<<8))
@@ -1574,11 +1578,13 @@ SQLITE_API int sqlite3_exec(
 #define SQLITE_CONSTRAINT_UNIQUE       (SQLITE_CONSTRAINT | (8<<8))
 #define SQLITE_CONSTRAINT_VTAB         (SQLITE_CONSTRAINT | (9<<8))
 #define SQLITE_CONSTRAINT_ROWID        (SQLITE_CONSTRAINT |(10<<8))
+#define SQLITE_CONSTRAINT_PINNED       (SQLITE_CONSTRAINT |(11<<8))
 #define SQLITE_NOTICE_RECOVER_WAL      (SQLITE_NOTICE | (1<<8))
 #define SQLITE_NOTICE_RECOVER_ROLLBACK (SQLITE_NOTICE | (2<<8))
 #define SQLITE_WARNING_AUTOINDEX       (SQLITE_WARNING | (1<<8))
 #define SQLITE_AUTH_USER               (SQLITE_AUTH | (1<<8))
 #define SQLITE_OK_LOAD_PERMANENTLY     (SQLITE_OK | (1<<8))
+#define SQLITE_OK_SYMLINK              (SQLITE_OK | (2<<8))
 
 /*
 ** CAPI3REF: Flags For File Open Operations
@@ -1607,6 +1613,7 @@ SQLITE_API int sqlite3_exec(
 #define SQLITE_OPEN_SHAREDCACHE      0x00020000  /* Ok for sqlite3_open_v2() */
 #define SQLITE_OPEN_PRIVATECACHE     0x00040000  /* Ok for sqlite3_open_v2() */
 #define SQLITE_OPEN_WAL              0x00080000  /* VFS only */
+#define SQLITE_OPEN_NOFOLLOW         0x01000000  /* Ok for sqlite3_open_v2() */
 
 /* Reserved:                         0x00F00000 */
 
@@ -2018,16 +2025,16 @@ struct sqlite3_io_methods {
 ** ^The [SQLITE_FCNTL_BUSYHANDLER]
 ** file-control may be invoked by SQLite on the database file handle
 ** shortly after it is opened in order to provide a custom VFS with access
-** to the connections busy-handler callback. The argument is of type (void **)
+** to the connection's busy-handler callback. The argument is of type (void**)
 ** - an array of two (void *) values. The first (void *) actually points
-** to a function of type (int (*)(void *)). In order to invoke the connections
+** to a function of type (int (*)(void *)). In order to invoke the connection's
 ** busy-handler, this function should be invoked with the second (void *) in
 ** the array as the only argument. If it returns non-zero, then the operation
 ** should be retried. If it returns zero, the custom VFS should abandon the
 ** current operation.
 **
 ** <li>[[SQLITE_FCNTL_TEMPFILENAME]]
-** ^Application can invoke the [SQLITE_FCNTL_TEMPFILENAME] file-control
+** ^Applications can invoke the [SQLITE_FCNTL_TEMPFILENAME] file-control
 ** to have SQLite generate a
 ** temporary filename using the same algorithm that is followed to generate
 ** temporary filenames for TEMP tables and other internal uses.  The
@@ -2140,12 +2147,18 @@ struct sqlite3_io_methods {
 ** not provide a mechanism to detect changes to MAIN only.  Also, the
 ** [sqlite3_total_changes()] interface responds to internal changes only and
 ** omits changes made by other database connections.  The
-** [PRAGMA data_version] command provide a mechanism to detect changes to
+** [PRAGMA data_version] command provides a mechanism to detect changes to
 ** a single attached database that occur due to other database connections,
 ** but omits changes implemented by the database connection on which it is
 ** called.  This file control is the only mechanism to detect changes that
 ** happen either internally or externally and that are associated with
 ** a particular attached database.
+**
+** <li>[[SQLITE_FCNTL_CKPT_DONE]]
+** The [SQLITE_FCNTL_CKPT_DONE] opcode is invoked from within a checkpoint
+** in wal mode after the client has finished copying pages from the wal
+** file to the database file, but before the *-shm file is updated to
+** record the fact that the pages have been checkpointed.
 ** </ul>
 */
 #define SQLITE_FCNTL_LOCKSTATE               1
@@ -2183,6 +2196,7 @@ struct sqlite3_io_methods {
 #define SQLITE_FCNTL_LOCK_TIMEOUT           34
 #define SQLITE_FCNTL_DATA_VERSION           35
 #define SQLITE_FCNTL_SIZE_LIMIT             36
+#define SQLITE_FCNTL_CKPT_DONE              37
 
 /* deprecated names */
 #define SQLITE_GET_LOCKPROXYFILE      SQLITE_FCNTL_GET_LOCKPROXYFILE
@@ -2228,10 +2242,10 @@ typedef struct sqlite3_api_routines sqlite3_api_routines;
 ** to 3 with SQLite [version 3.7.6] on [dateof:3.7.6].  Additional fields
 ** may be appended to the sqlite3_vfs object and the iVersion value
 ** may increase again in future versions of SQLite.
-** Note that the structure
-** of the sqlite3_vfs object changes in the transition from
+** Note that due to an oversight, the structure
+** of the sqlite3_vfs object changed in the transition from
 ** SQLite [version 3.5.9] to [version 3.6.0] on [dateof:3.6.0]
-** and yet the iVersion field was not modified.
+** and yet the iVersion field was not increased.
 **
 ** The szOsFile field is the size of the subclassed [sqlite3_file]
 ** structure used by this VFS.  mxPathname is the maximum length of
@@ -2322,7 +2336,7 @@ typedef struct sqlite3_api_routines sqlite3_api_routines;
 ** for exclusive access.
 **
 ** ^At least szOsFile bytes of memory are allocated by SQLite
-** to hold the  [sqlite3_file] structure passed as the third
+** to hold the [sqlite3_file] structure passed as the third
 ** argument to xOpen.  The xOpen method does not have to
 ** allocate the structure; it should just fill it in.  Note that
 ** the xOpen method must set the sqlite3_file.pMethods to either
@@ -2335,8 +2349,14 @@ typedef struct sqlite3_api_routines sqlite3_api_routines;
 ** ^The flags argument to xAccess() may be [SQLITE_ACCESS_EXISTS]
 ** to test for the existence of a file, or [SQLITE_ACCESS_READWRITE] to
 ** test whether a file is readable and writable, or [SQLITE_ACCESS_READ]
-** to test whether a file is at least readable.   The file can be a
-** directory.
+** to test whether a file is at least readable.  The SQLITE_ACCESS_READ
+** flag is never actually used and is not implemented in the built-in
+** VFSes of SQLite.  The file is named by the second argument and can be a
+** directory. The xAccess method returns [SQLITE_OK] on success or some
+** non-zero error code if there is an I/O error or if the name of
+** the file given in the second argument is illegal.  If SQLITE_OK
+** is returned, then non-zero or zero is written into *pResOut to indicate
+** whether or not the file is accessible.  
 **
 ** ^SQLite will always allocate at least mxPathname+1 bytes for the
 ** output buffer xFullPathname.  The exact size of the output buffer
@@ -2653,7 +2673,7 @@ SQLITE_API int sqlite3_db_config(sqlite3*, int op, ...);
 ** that causes the corresponding memory allocation to fail.
 **
 ** The xInit method initializes the memory allocator.  For example,
-** it might allocate any require mutexes or initialize internal data
+** it might allocate any required mutexes or initialize internal data
 ** structures.  The xShutdown method is invoked (indirectly) by
 ** [sqlite3_shutdown()] and should deallocate any resources acquired
 ** by xInit.  The pAppData pointer is used as the only parameter to
@@ -2775,6 +2795,7 @@ struct sqlite3_mem_methods {
 ** memory allocation statistics. ^(When memory allocation statistics are
 ** disabled, the following SQLite interfaces become non-operational:
 **   <ul>
+**   <li> [sqlite3_hard_heap_limit64()]
 **   <li> [sqlite3_memory_used()]
 **   <li> [sqlite3_memory_highwater()]
 **   <li> [sqlite3_soft_heap_limit64()]
@@ -2793,7 +2814,7 @@ struct sqlite3_mem_methods {
 ** <dd> ^The SQLITE_CONFIG_PAGECACHE option specifies a memory pool
 ** that SQLite can use for the database page cache with the default page
 ** cache implementation.  
-** This configuration option is a no-op if an application-define page
+** This configuration option is a no-op if an application-defined page
 ** cache implementation is loaded using the [SQLITE_CONFIG_PCACHE2].
 ** ^There are three arguments to SQLITE_CONFIG_PAGECACHE: A pointer to
 ** 8-byte aligned memory (pMem), the size of each page cache line (sz),
@@ -3126,6 +3147,17 @@ struct sqlite3_mem_methods {
 ** following this call.  The second parameter may be a NULL pointer, in
 ** which case the trigger setting is not reported back. </dd>
 **
+** [[SQLITE_DBCONFIG_ENABLE_VIEW]]
+** <dt>SQLITE_DBCONFIG_ENABLE_VIEW</dt>
+** <dd> ^This option is used to enable or disable [CREATE VIEW | views].
+** There should be two additional arguments.
+** The first argument is an integer which is 0 to disable views,
+** positive to enable views or negative to leave the setting unchanged.
+** The second parameter is a pointer to an integer into which
+** is written 0 or 1 to indicate whether views are disabled or enabled
+** following this call.  The second parameter may be a NULL pointer, in
+** which case the view setting is not reported back. </dd>
+**
 ** [[SQLITE_DBCONFIG_ENABLE_FTS3_TOKENIZER]]
 ** <dt>SQLITE_DBCONFIG_ENABLE_FTS3_TOKENIZER</dt>
 ** <dd> ^This option is used to enable or disable the
@@ -3237,6 +3269,7 @@ struct sqlite3_mem_methods {
 ** features include but are not limited to the following:
 ** <ul>
 ** <li> The [PRAGMA writable_schema=ON] statement.
+** <li> The [PRAGMA journal_mode=OFF] statement.
 ** <li> Writes to the [sqlite_dbpage] virtual table.
 ** <li> Direct writes to [shadow tables].
 ** </ul>
@@ -3252,6 +3285,77 @@ struct sqlite3_mem_methods {
 ** integer into which is written 0 or 1 to indicate whether the writable_schema
 ** is enabled or disabled following this call.
 ** </dd>
+**
+** [[SQLITE_DBCONFIG_LEGACY_ALTER_TABLE]]
+** <dt>SQLITE_DBCONFIG_LEGACY_ALTER_TABLE</dt>
+** <dd>The SQLITE_DBCONFIG_LEGACY_ALTER_TABLE option activates or deactivates
+** the legacy behavior of the [ALTER TABLE RENAME] command such it
+** behaves as it did prior to [version 3.24.0] (2018-06-04).  See the
+** "Compatibility Notice" on the [ALTER TABLE RENAME documentation] for
+** additional information. This feature can also be turned on and off
+** using the [PRAGMA legacy_alter_table] statement.
+** </dd>
+**
+** [[SQLITE_DBCONFIG_DQS_DML]]
+** <dt>SQLITE_DBCONFIG_DQS_DML</td>
+** <dd>The SQLITE_DBCONFIG_DQS_DML option activates or deactivates
+** the legacy [double-quoted string literal] misfeature for DML statements
+** only, that is DELETE, INSERT, SELECT, and UPDATE statements. The
+** default value of this setting is determined by the [-DSQLITE_DQS]
+** compile-time option.
+** </dd>
+**
+** [[SQLITE_DBCONFIG_DQS_DDL]]
+** <dt>SQLITE_DBCONFIG_DQS_DDL</td>
+** <dd>The SQLITE_DBCONFIG_DQS option activates or deactivates
+** the legacy [double-quoted string literal] misfeature for DDL statements,
+** such as CREATE TABLE and CREATE INDEX. The
+** default value of this setting is determined by the [-DSQLITE_DQS]
+** compile-time option.
+** </dd>
+**
+** [[SQLITE_DBCONFIG_TRUSTED_SCHEMA]]
+** <dt>SQLITE_DBCONFIG_TRUSTED_SCHEMA</td>
+** <dd>The SQLITE_DBCONFIG_TRUSTED_SCHEMA option tells SQLite to
+** assume that database schemas (the contents of the [sqlite_master] tables)
+** are untainted by malicious content.
+** When the SQLITE_DBCONFIG_TRUSTED_SCHEMA option is disabled, SQLite
+** takes additional defensive steps to protect the application from harm
+** including:
+** <ul>
+** <li> Prohibit the use of SQL functions inside triggers, views,
+** CHECK constraints, DEFAULT clauses, expression indexes, 
+** partial indexes, or generated columns
+** unless those functions are tagged with [SQLITE_INNOCUOUS].
+** <li> Prohibit the use of virtual tables inside of triggers or views
+** unless those virtual tables are tagged with [SQLITE_VTAB_INNOCUOUS].
+** </ul>
+** This setting defaults to "on" for legacy compatibility, however
+** all applications are advised to turn it off if possible. This setting
+** can also be controlled using the [PRAGMA trusted_schema] statement.
+** </dd>
+**
+** [[SQLITE_DBCONFIG_LEGACY_FILE_FORMAT]]
+** <dt>SQLITE_DBCONFIG_LEGACY_FILE_FORMAT</td>
+** <dd>The SQLITE_DBCONFIG_LEGACY_FILE_FORMAT option activates or deactivates
+** the legacy file format flag.  When activated, this flag causes all newly
+** created database file to have a schema format version number (the 4-byte
+** integer found at offset 44 into the database header) of 1.  This in turn
+** means that the resulting database file will be readable and writable by
+** any SQLite version back to 3.0.0 ([dateof:3.0.0]).  Without this setting,
+** newly created databases are generally not understandable by SQLite versions
+** prior to 3.3.0 ([dateof:3.3.0]).  As these words are written, there
+** is now scarcely any need to generated database files that are compatible 
+** all the way back to version 3.0.0, and so this setting is of little
+** practical use, but is provided so that SQLite can continue to claim the
+** ability to generate new database files that are compatible with  version
+** 3.0.0.
+** <p>Note that when the SQLITE_DBCONFIG_LEGACY_FILE_FORMAT setting is on,
+** the [VACUUM] command will fail with an obscure error when attempting to
+** process a table with generated columns and a descending index.  This is
+** not considered a bug since SQLite versions 3.3.0 and earlier do not support
+** either generated columns or decending indexes.
+** </dd>
 ** </dl>
 */
 #define SQLITE_DBCONFIG_MAINDBNAME            1000 /* const char* */
@@ -3266,7 +3370,13 @@ struct sqlite3_mem_methods {
 #define SQLITE_DBCONFIG_RESET_DATABASE        1009 /* int int* */
 #define SQLITE_DBCONFIG_DEFENSIVE             1010 /* int int* */
 #define SQLITE_DBCONFIG_WRITABLE_SCHEMA       1011 /* int int* */
-#define SQLITE_DBCONFIG_MAX                   1011 /* Largest DBCONFIG */
+#define SQLITE_DBCONFIG_LEGACY_ALTER_TABLE    1012 /* int int* */
+#define SQLITE_DBCONFIG_DQS_DML               1013 /* int int* */
+#define SQLITE_DBCONFIG_DQS_DDL               1014 /* int int* */
+#define SQLITE_DBCONFIG_ENABLE_VIEW           1015 /* int int* */
+#define SQLITE_DBCONFIG_LEGACY_FILE_FORMAT    1016 /* int int* */
+#define SQLITE_DBCONFIG_TRUSTED_SCHEMA        1017 /* int int* */
+#define SQLITE_DBCONFIG_MAX                   1017 /* Largest DBCONFIG */
 
 /*
 ** CAPI3REF: Enable Or Disable Extended Result Codes
@@ -3472,7 +3582,7 @@ SQLITE_API int sqlite3_total_changes(sqlite3*);
 ** ^The sqlite3_interrupt(D) call is in effect until all currently running
 ** SQL statements on [database connection] D complete.  ^Any new SQL statements
 ** that are started after the sqlite3_interrupt() call and before the 
-** running statements reaches zero are interrupted as if they had been
+** running statement count reaches zero are interrupted as if they had been
 ** running prior to the sqlite3_interrupt() call.  ^New SQL statements
 ** that are started after the running statement count reaches zero are
 ** not effected by the sqlite3_interrupt().
@@ -3640,9 +3750,9 @@ SQLITE_API int sqlite3_busy_timeout(sqlite3*, int ms);
 **        Cindy       | 21
 ** </pre></blockquote>
 **
-** There are two column (M==2) and three rows (N==3).  Thus the
+** There are two columns (M==2) and three rows (N==3).  Thus the
 ** result table has 8 entries.  Suppose the result table is stored
-** in an array names azResult.  Then azResult holds this content:
+** in an array named azResult.  Then azResult holds this content:
 **
 ** <blockquote><pre>
 **        azResult&#91;0] = "Name";
@@ -3735,7 +3845,7 @@ SQLITE_API char *sqlite3_vsnprintf(int,char*,const char*, va_list);
 **
 ** The SQLite core uses these three routines for all of its own
 ** internal memory allocation needs. "Core" in the previous sentence
-** does not include operating-system specific VFS implementation.  The
+** does not include operating-system specific [VFS] implementation.  The
 ** Windows VFS uses native malloc() and free() for some operations.
 **
 ** ^The sqlite3_malloc() routine returns a pointer to a block
@@ -3796,19 +3906,6 @@ SQLITE_API char *sqlite3_vsnprintf(int,char*,const char*, va_list);
 ** 4 byte boundary if the [SQLITE_4_BYTE_ALIGNED_MALLOC] compile-time
 ** option is used.
 **
-** In SQLite version 3.5.0 and 3.5.1, it was possible to define
-** the SQLITE_OMIT_MEMORY_ALLOCATION which would cause the built-in
-** implementation of these routines to be omitted.  That capability
-** is no longer provided.  Only built-in memory allocators can be used.
-**
-** Prior to SQLite version 3.7.10, the Windows OS interface layer called
-** the system malloc() and free() directly when converting
-** filenames between the UTF-8 encoding used by SQLite
-** and whatever filename encoding is used by the particular Windows
-** installation.  Memory allocation errors were detected, but
-** they were reported back as [SQLITE_CANTOPEN] or
-** [SQLITE_IOERR] rather than [SQLITE_NOMEM].
-**
 ** The pointer arguments to [sqlite3_free()] and [sqlite3_realloc()]
 ** must be either NULL or else pointers obtained from a prior
 ** invocation of [sqlite3_malloc()] or [sqlite3_realloc()] that have
@@ -3857,7 +3954,7 @@ SQLITE_API sqlite3_int64 sqlite3_memory_highwater(int resetFlag);
 ** SQLite contains a high-quality pseudo-random number generator (PRNG) used to
 ** select random [ROWID | ROWIDs] when inserting new records into a table that
 ** already uses the largest possible [ROWID].  The PRNG is also used for
-** the build-in random() and randomblob() SQL functions.  This interface allows
+** the built-in random() and randomblob() SQL functions.  This interface allows
 ** applications to access the same PRNG for other purposes.
 **
 ** ^A call to this routine stores N bytes of randomness into buffer P.
@@ -4231,10 +4328,8 @@ SQLITE_API void sqlite3_progress_handler(sqlite3*, int, int(*)(void*), void*);
 ** The sqlite3_open_v2() interface works like sqlite3_open()
 ** except that it accepts two additional parameters for additional control
 ** over the new database connection.  ^(The flags parameter to
-** sqlite3_open_v2() can take one of
-** the following three values, optionally combined with the 
-** [SQLITE_OPEN_NOMUTEX], [SQLITE_OPEN_FULLMUTEX], [SQLITE_OPEN_SHAREDCACHE],
-** [SQLITE_OPEN_PRIVATECACHE], and/or [SQLITE_OPEN_URI] flags:)^
+** sqlite3_open_v2() must include, at a minimum, one of the following
+** three flag combinations:)^
 **
 ** <dl>
 ** ^(<dt>[SQLITE_OPEN_READONLY]</dt>
@@ -4252,23 +4347,51 @@ SQLITE_API void sqlite3_progress_handler(sqlite3*, int, int(*)(void*), void*);
 ** sqlite3_open() and sqlite3_open16().</dd>)^
 ** </dl>
 **
+** In addition to the required flags, the following optional flags are
+** also supported:
+**
+** <dl>
+** ^(<dt>[SQLITE_OPEN_URI]</dt>
+** <dd>The filename can be interpreted as a URI if this flag is set.</dd>)^
+**
+** ^(<dt>[SQLITE_OPEN_MEMORY]</dt>
+** <dd>The database will be opened as an in-memory database.  The database
+** is named by the "filename" argument for the purposes of cache-sharing,
+** if shared cache mode is enabled, but the "filename" is otherwise ignored.
+** </dd>)^
+**
+** ^(<dt>[SQLITE_OPEN_NOMUTEX]</dt>
+** <dd>The new database connection will use the "multi-thread"
+** [threading mode].)^  This means that separate threads are allowed
+** to use SQLite at the same time, as long as each thread is using
+** a different [database connection].
+**
+** ^(<dt>[SQLITE_OPEN_FULLMUTEX]</dt>
+** <dd>The new database connection will use the "serialized"
+** [threading mode].)^  This means the multiple threads can safely
+** attempt to use the same database connection at the same time.
+** (Mutexes will block any actual concurrency, but in this mode
+** there is no harm in trying.)
+**
+** ^(<dt>[SQLITE_OPEN_SHAREDCACHE]</dt>
+** <dd>The database is opened [shared cache] enabled, overriding
+** the default shared cache setting provided by
+** [sqlite3_enable_shared_cache()].)^
+**
+** ^(<dt>[SQLITE_OPEN_PRIVATECACHE]</dt>
+** <dd>The database is opened [shared cache] disabled, overriding
+** the default shared cache setting provided by
+** [sqlite3_enable_shared_cache()].)^
+**
+** [[OPEN_NOFOLLOW]] ^(<dt>[SQLITE_OPEN_NOFOLLOW]</dt>
+** <dd>The database filename is not allowed to be a symbolic link</dd>
+** </dl>)^
+**
 ** If the 3rd parameter to sqlite3_open_v2() is not one of the
-** combinations shown above optionally combined with other
+** required combinations shown above optionally combined with other
 ** [SQLITE_OPEN_READONLY | SQLITE_OPEN_* bits]
 ** then the behavior is undefined.
 **
-** ^If the [SQLITE_OPEN_NOMUTEX] flag is set, then the database connection
-** opens in the multi-thread [threading mode] as long as the single-thread
-** mode has not been set at compile-time or start-time.  ^If the
-** [SQLITE_OPEN_FULLMUTEX] flag is set then the database connection opens
-** in the serialized [threading mode] unless single-thread was
-** previously selected at compile-time or start-time.
-** ^The [SQLITE_OPEN_SHAREDCACHE] flag causes the database connection to be
-** eligible to use [shared cache mode], regardless of whether or not shared
-** cache is enabled using [sqlite3_enable_shared_cache()].  ^The
-** [SQLITE_OPEN_PRIVATECACHE] flag causes the database connection to not
-** participate in [shared cache mode] even if it is enabled.
-**
 ** ^The fourth parameter to sqlite3_open_v2() is the name of the
 ** [sqlite3_vfs] object that defines the operating system interface that
 ** the new database connection should use.  ^If the fourth parameter is
@@ -4448,17 +4571,16 @@ SQLITE_API int sqlite3_open_v2(
 /*
 ** CAPI3REF: Obtain Values For URI Parameters
 **
-** These are utility routines, useful to VFS implementations, that check
-** to see if a database file was a URI that contained a specific query 
+** These are utility routines, useful to [VFS|custom VFS implementations],
+** that check if a database file was a URI that contained a specific query 
 ** parameter, and if so obtains the value of that query parameter.
 **
 ** If F is the database filename pointer passed into the xOpen() method of 
-** a VFS implementation when the flags parameter to xOpen() has one or 
-** more of the [SQLITE_OPEN_URI] or [SQLITE_OPEN_MAIN_DB] bits set and
-** P is the name of the query parameter, then
+** a VFS implementation or it is the return value of [sqlite3_db_filename()]
+** and if P is the name of the query parameter, then
 ** sqlite3_uri_parameter(F,P) returns the value of the P
 ** parameter if it exists or a NULL pointer if P does not appear as a 
-** query parameter on F.  If P is a query parameter of F
+** query parameter on F.  If P is a query parameter of F and it
 ** has no explicit value, then sqlite3_uri_parameter(F,P) returns
 ** a pointer to an empty string.
 **
@@ -4470,25 +4592,72 @@ SQLITE_API int sqlite3_open_v2(
 ** sqlite3_uri_boolean(F,P,B) routines returns false (0) if the value of
 ** query parameter P is one of "no", "false", or "off" in any case or
 ** if the value begins with a numeric zero.  If P is not a query
-** parameter on F or if the value of P is does not match any of the
+** parameter on F or if the value of P does not match any of the
 ** above, then sqlite3_uri_boolean(F,P,B) returns (B!=0).
 **
 ** The sqlite3_uri_int64(F,P,D) routine converts the value of P into a
 ** 64-bit signed integer and returns that integer, or D if P does not
 ** exist.  If the value of P is something other than an integer, then
 ** zero is returned.
+**
+** The sqlite3_uri_key(F,N) returns a pointer to the name (not
+** the value) of the N-th query parameter for filename F, or a NULL
+** pointer if N is less than zero or greater than the number of query
+** parameters minus 1.  The N value is zero-based so N should be 0 to obtain
+** the name of the first query parameter, 1 for the second parameter, and
+** so forth.
 ** 
 ** If F is a NULL pointer, then sqlite3_uri_parameter(F,P) returns NULL and
 ** sqlite3_uri_boolean(F,P,B) returns B.  If F is not a NULL pointer and
-** is not a database file pathname pointer that SQLite passed into the xOpen
-** VFS method, then the behavior of this routine is undefined and probably
-** undesirable.
+** is not a database file pathname pointer that the SQLite core passed
+** into the xOpen VFS method, then the behavior of this routine is undefined
+** and probably undesirable.
+**
+** Beginning with SQLite [version 3.31.0] ([dateof:3.31.0]) the input F
+** parameter can also be the name of a rollback journal file or WAL file
+** in addition to the main database file.  Prior to version 3.31.0, these
+** routines would only work if F was the name of the main database file.
+** When the F parameter is the name of the rollback journal or WAL file,
+** it has access to all the same query parameters as were found on the
+** main database file.
 **
 ** See the [URI filename] documentation for additional information.
 */
 SQLITE_API const char *sqlite3_uri_parameter(const char *zFilename, const char *zParam);
 SQLITE_API int sqlite3_uri_boolean(const char *zFile, const char *zParam, int bDefault);
 SQLITE_API sqlite3_int64 sqlite3_uri_int64(const char*, const char*, sqlite3_int64);
+SQLITE_API const char *sqlite3_uri_key(const char *zFilename, int N);
+
+/*
+** CAPI3REF:  Translate filenames
+**
+** These routines are available to [VFS|custom VFS implementations] for
+** translating filenames between the main database file, the journal file,
+** and the WAL file.
+**
+** If F is the name of an sqlite database file, journal file, or WAL file
+** passed by the SQLite core into the VFS, then sqlite3_filename_database(F)
+** returns the name of the corresponding database file.
+**
+** If F is the name of an sqlite database file, journal file, or WAL file
+** passed by the SQLite core into the VFS, or if F is a database filename
+** obtained from [sqlite3_db_filename()], then sqlite3_filename_journal(F)
+** returns the name of the corresponding rollback journal file.
+**
+** If F is the name of an sqlite database file, journal file, or WAL file
+** that was passed by the SQLite core into the VFS, or if F is a database
+** filename obtained from [sqlite3_db_filename()], then
+** sqlite3_filename_wal(F) returns the name of the corresponding
+** WAL file.
+**
+** In all of the above, if F is not the name of a database, journal or WAL
+** filename passed into the VFS from the SQLite core and F is not the
+** return value from [sqlite3_db_filename()], then the result is
+** undefined and is likely a memory access violation.
+*/
+SQLITE_API const char *sqlite3_filename_database(const char*);
+SQLITE_API const char *sqlite3_filename_journal(const char*);
+SQLITE_API const char *sqlite3_filename_wal(const char*);
 
 
 /*
@@ -4807,15 +4976,15 @@ SQLITE_API int sqlite3_limit(sqlite3*, int id, int newVal);
 ** </li>
 **
 ** <li>
-** ^If the specific value bound to [parameter | host parameter] in the 
+** ^If the specific value bound to a [parameter | host parameter] in the 
 ** WHERE clause might influence the choice of query plan for a statement,
 ** then the statement will be automatically recompiled, as if there had been 
-** a schema change, on the first  [sqlite3_step()] call following any change
+** a schema change, on the first [sqlite3_step()] call following any change
 ** to the [sqlite3_bind_text | bindings] of that [parameter]. 
-** ^The specific value of WHERE-clause [parameter] might influence the 
+** ^The specific value of a WHERE-clause [parameter] might influence the 
 ** choice of query plan if the parameter is the left-hand side of a [LIKE]
 ** or [GLOB] operator or if the parameter is compared to an indexed column
-** and the [SQLITE_ENABLE_STAT3] compile-time option is enabled.
+** and the [SQLITE_ENABLE_STAT4] compile-time option is enabled.
 ** </li>
 ** </ol>
 **
@@ -5321,7 +5490,7 @@ SQLITE_API const void *sqlite3_column_name16(sqlite3_stmt*, int N);
 **
 ** ^If the Nth column returned by the statement is an expression or
 ** subquery and is not a column value, then all of these functions return
-** NULL.  ^These routine might also return NULL if a memory allocation error
+** NULL.  ^These routines might also return NULL if a memory allocation error
 ** occurs.  ^Otherwise, they return the name of the attached database, table,
 ** or column that query result column was extracted from.
 **
@@ -5331,10 +5500,6 @@ SQLITE_API const void *sqlite3_column_name16(sqlite3_stmt*, int N);
 ** ^These APIs are only available if the library was compiled with the
 ** [SQLITE_ENABLE_COLUMN_METADATA] C-preprocessor symbol.
 **
-** If two or more threads call one or more of these routines against the same
-** prepared statement and column at the same time then the results are
-** undefined.
-**
 ** If two or more threads call one or more
 ** [sqlite3_column_database_name | column metadata interfaces]
 ** for the same [prepared statement] and result column
@@ -5471,7 +5636,7 @@ SQLITE_API int sqlite3_step(sqlite3_stmt*);
 ** ^The sqlite3_data_count(P) interface returns the number of columns in the
 ** current row of the result set of [prepared statement] P.
 ** ^If prepared statement P does not have results ready to return
-** (via calls to the [sqlite3_column_int | sqlite3_column_*()] of
+** (via calls to the [sqlite3_column_int | sqlite3_column()] family of
 ** interfaces) then sqlite3_data_count(P) returns 0.
 ** ^The sqlite3_data_count(P) routine also returns 0 if P is a NULL pointer.
 ** ^The sqlite3_data_count(P) routine returns 0 if the previous call to
@@ -5795,8 +5960,6 @@ SQLITE_API int sqlite3_reset(sqlite3_stmt *pStmt);
 /*
 ** CAPI3REF: Create Or Redefine SQL Functions
 ** KEYWORDS: {function creation routines}
-** KEYWORDS: {application-defined SQL function}
-** KEYWORDS: {application-defined SQL functions}
 ** METHOD: sqlite3
 **
 ** ^These functions (collectively known as "function creation routines")
@@ -5850,6 +6013,23 @@ SQLITE_API int sqlite3_reset(sqlite3_stmt *pStmt);
 ** perform additional optimizations on deterministic functions, so use
 ** of the [SQLITE_DETERMINISTIC] flag is recommended where possible.
 **
+** ^The fourth parameter may also optionally include the [SQLITE_DIRECTONLY]
+** flag, which if present prevents the function from being invoked from
+** within VIEWs, TRIGGERs, CHECK constraints, generated column expressions,
+** index expressions, or the WHERE clause of partial indexes.
+**
+** <span style="background-color:#ffff90;">
+** For best security, the [SQLITE_DIRECTONLY] flag is recommended for
+** all application-defined SQL functions that do not need to be
+** used inside of triggers, view, CHECK constraints, or other elements of
+** the database schema.  This flags is especially recommended for SQL 
+** functions that have side effects or reveal internal application state.
+** Without this flag, an attacker might be able to modify the schema of
+** a database file to include invocations of the function with parameters
+** chosen by the attacker, which the application will then execute when
+** the database file is opened and read.
+** </span>
+**
 ** ^(The fifth parameter is an arbitrary pointer.  The implementation of the
 ** function can gain access to this pointer using [sqlite3_user_data()].)^
 **
@@ -5966,8 +6146,68 @@ SQLITE_API int sqlite3_create_window_function(
 ** [SQLITE_UTF8 | preferred text encoding] as the fourth argument
 ** to [sqlite3_create_function()], [sqlite3_create_function16()], or
 ** [sqlite3_create_function_v2()].
+**
+** <dl>
+** [[SQLITE_DETERMINISTIC]] <dt>SQLITE_DETERMINISTIC</dt><dd>
+** The SQLITE_DETERMINISTIC flag means that the new function always gives
+** the same output when the input parameters are the same.
+** The [abs|abs() function] is deterministic, for example, but
+** [randomblob|randomblob()] is not.  Functions must
+** be deterministic in order to be used in certain contexts such as
+** with the WHERE clause of [partial indexes] or in [generated columns].
+** SQLite might also optimize deterministic functions by factoring them
+** out of inner loops.
+** </dd>
+** 
+** [[SQLITE_DIRECTONLY]] <dt>SQLITE_DIRECTONLY</dt><dd>
+** The SQLITE_DIRECTONLY flag means that the function may only be invoked
+** from top-level SQL, and cannot be used in VIEWs or TRIGGERs nor in 
+** schema structures such as [CHECK constraints], [DEFAULT clauses],
+** [expression indexes], [partial indexes], or [generated columns].
+** The SQLITE_DIRECTONLY flags is a security feature which is recommended
+** for all [application-defined SQL functions], and especially for functions
+** that have side-effects or that could potentially leak sensitive
+** information.
+** </dd>
+**
+** [[SQLITE_INNOCUOUS]] <dt>SQLITE_INNOCUOUS</dt><dd>
+** The SQLITE_INNOCUOUS flag means that the function is unlikely
+** to cause problems even if misused.  An innocuous function should have
+** no side effects and should not depend on any values other than its
+** input parameters. The [abs|abs() function] is an example of an
+** innocuous function.
+** The [load_extension() SQL function] is not innocuous because of its
+** side effects.
+** <p> SQLITE_INNOCUOUS is similar to SQLITE_DETERMINISTIC, but is not
+** exactly the same.  The [random|random() function] is an example of a
+** function that is innocuous but not deterministic.
+** <p>Some heightened security settings
+** ([SQLITE_DBCONFIG_TRUSTED_SCHEMA] and [PRAGMA trusted_schema=OFF])
+** disable the use of SQL functions inside views and triggers and in
+** schema structures such as [CHECK constraints], [DEFAULT clauses],
+** [expression indexes], [partial indexes], and [generated columns] unless
+** the function is tagged with SQLITE_INNOCUOUS.  Most built-in functions
+** are innocuous.  Developers are advised to avoid using the
+** SQLITE_INNOCUOUS flag for application-defined functions unless the
+** function has been carefully audited and found to be free of potentially
+** security-adverse side-effects and information-leaks.
+** </dd>
+**
+** [[SQLITE_SUBTYPE]] <dt>SQLITE_SUBTYPE</dt><dd>
+** The SQLITE_SUBTYPE flag indicates to SQLite that a function may call
+** [sqlite3_value_subtype()] to inspect the sub-types of its arguments.
+** Specifying this flag makes no difference for scalar or aggregate user
+** functions. However, if it is not specified for a user-defined window
+** function, then any sub-types belonging to arguments passed to the window
+** function may be discarded before the window function is called (i.e.
+** sqlite3_value_subtype() will always return 0).
+** </dd>
+** </dl>
 */
-#define SQLITE_DETERMINISTIC    0x800
+#define SQLITE_DETERMINISTIC    0x000000800
+#define SQLITE_DIRECTONLY       0x000080000
+#define SQLITE_SUBTYPE          0x000100000
+#define SQLITE_INNOCUOUS        0x000200000
 
 /*
 ** CAPI3REF: Deprecated Functions
@@ -6026,8 +6266,8 @@ SQLITE_API SQLITE_DEPRECATED int sqlite3_memory_alarm(void(*)(void*,sqlite3_int6
 **
 ** These routines extract type, size, and content information from
 ** [protected sqlite3_value] objects.  Protected sqlite3_value objects
-** are used to pass parameter information into implementation of
-** [application-defined SQL functions] and [virtual tables].
+** are used to pass parameter information into the functions that
+** implement [application-defined SQL functions] and [virtual tables].
 **
 ** These routines work only with [protected sqlite3_value] objects.
 ** Any attempt to use these routines on an [unprotected sqlite3_value]
@@ -6084,7 +6324,7 @@ SQLITE_API SQLITE_DEPRECATED int sqlite3_memory_alarm(void(*)(void*,sqlite3_int6
 ** ^The sqlite3_value_frombind(X) interface returns non-zero if the
 ** value X originated from one of the [sqlite3_bind_int|sqlite3_bind()]
 ** interfaces.  ^If X comes from an SQL literal value, or a table column,
-** and expression, then sqlite3_value_frombind(X) returns zero.
+** or an expression, then sqlite3_value_frombind(X) returns zero.
 **
 ** Please pay particular attention to the fact that the pointer returned
 ** from [sqlite3_value_blob()], [sqlite3_value_text()], or
@@ -6170,8 +6410,8 @@ SQLITE_API void sqlite3_value_free(sqlite3_value*);
 ** routine to allocate memory for storing their state.
 **
 ** ^The first time the sqlite3_aggregate_context(C,N) routine is called 
-** for a particular aggregate function, SQLite
-** allocates N of memory, zeroes out that memory, and returns a pointer
+** for a particular aggregate function, SQLite allocates
+** N bytes of memory, zeroes out that memory, and returns a pointer
 ** to the new memory. ^On second and subsequent calls to
 ** sqlite3_aggregate_context() for the same aggregate function instance,
 ** the same buffer is returned.  Sqlite3_aggregate_context() is normally
@@ -6188,7 +6428,7 @@ SQLITE_API void sqlite3_value_free(sqlite3_value*);
 **
 ** ^(The amount of space allocated by sqlite3_aggregate_context(C,N) is
 ** determined by the N parameter on first successful call.  Changing the
-** value of N in subsequent call to sqlite3_aggregate_context() within
+** value of N in any subsequents call to sqlite3_aggregate_context() within
 ** the same aggregate function instance will not resize the memory
 ** allocation.)^  Within the xFinal callback, it is customary to set
 ** N=0 in calls to sqlite3_aggregate_context(C,N) so that no 
@@ -6499,7 +6739,7 @@ SQLITE_API void sqlite3_result_subtype(sqlite3_context*,unsigned int);
 ** <li> [SQLITE_UTF16_ALIGNED].
 ** </ul>)^
 ** ^The eTextRep argument determines the encoding of strings passed
-** to the collating function callback, xCallback.
+** to the collating function callback, xCompare.
 ** ^The [SQLITE_UTF16] and [SQLITE_UTF16_ALIGNED] values for eTextRep
 ** force strings to be UTF16 with native byte order.
 ** ^The [SQLITE_UTF16_ALIGNED] value for eTextRep forces strings to begin
@@ -6508,18 +6748,19 @@ SQLITE_API void sqlite3_result_subtype(sqlite3_context*,unsigned int);
 ** ^The fourth argument, pArg, is an application data pointer that is passed
 ** through as the first argument to the collating function callback.
 **
-** ^The fifth argument, xCallback, is a pointer to the collating function.
+** ^The fifth argument, xCompare, is a pointer to the collating function.
 ** ^Multiple collating functions can be registered using the same name but
 ** with different eTextRep parameters and SQLite will use whichever
 ** function requires the least amount of data transformation.
-** ^If the xCallback argument is NULL then the collating function is
+** ^If the xCompare argument is NULL then the collating function is
 ** deleted.  ^When all collating functions having the same name are deleted,
 ** that collation is no longer usable.
 **
 ** ^The collating function callback is invoked with a copy of the pArg 
 ** application data pointer and with two strings in the encoding specified
-** by the eTextRep argument.  The collating function must return an
-** integer that is negative, zero, or positive
+** by the eTextRep argument.  The two integer parameters to the collating
+** function callback are the length of the two strings, in bytes. The collating
+** function must return an integer that is negative, zero, or positive
 ** if the first string is less than, equal to, or greater than the second,
 ** respectively.  A collating function must always return the same answer
 ** given the same inputs.  If two or more collating functions are registered
@@ -6536,7 +6777,7 @@ SQLITE_API void sqlite3_result_subtype(sqlite3_context*,unsigned int);
 ** </ol>
 **
 ** If a collating function fails any of the above constraints and that
-** collating function is  registered and used, then the behavior of SQLite
+** collating function is registered and used, then the behavior of SQLite
 ** is undefined.
 **
 ** ^The sqlite3_create_collation_v2() works like sqlite3_create_collation()
@@ -6863,16 +7104,31 @@ SQLITE_API sqlite3 *sqlite3_db_handle(sqlite3_stmt*);
 ** CAPI3REF: Return The Filename For A Database Connection
 ** METHOD: sqlite3
 **
-** ^The sqlite3_db_filename(D,N) interface returns a pointer to a filename
-** associated with database N of connection D.  ^The main database file
-** has the name "main".  If there is no attached database N on the database
+** ^The sqlite3_db_filename(D,N) interface returns a pointer to the filename
+** associated with database N of connection D.
+** ^If there is no attached database N on the database
 ** connection D, or if database N is a temporary or in-memory database, then
 ** this function will return either a NULL pointer or an empty string.
 **
+** ^The string value returned by this routine is owned and managed by
+** the database connection.  ^The value will be valid until the database N
+** is [DETACH]-ed or until the database connection closes.
+**
 ** ^The filename returned by this function is the output of the
 ** xFullPathname method of the [VFS].  ^In other words, the filename
 ** will be an absolute pathname, even if the filename used
 ** to open the database originally was a URI or relative pathname.
+**
+** If the filename pointer returned by this routine is not NULL, then it
+** can be used as the filename input parameter to these routines:
+** <ul>
+** <li> [sqlite3_uri_parameter()]
+** <li> [sqlite3_uri_boolean()]
+** <li> [sqlite3_uri_int64()]
+** <li> [sqlite3_filename_database()]
+** <li> [sqlite3_filename_journal()]
+** <li> [sqlite3_filename_wal()]
+** </ul>
 */
 SQLITE_API const char *sqlite3_db_filename(sqlite3 *db, const char *zDbName);
 
@@ -7022,15 +7278,19 @@ SQLITE_API void *sqlite3_update_hook(
 **
 ** ^(The cache sharing mode set by this interface effects all subsequent
 ** calls to [sqlite3_open()], [sqlite3_open_v2()], and [sqlite3_open16()].
-** Existing database connections continue use the sharing mode
+** Existing database connections continue to use the sharing mode
 ** that was in effect at the time they were opened.)^
 **
 ** ^(This routine returns [SQLITE_OK] if shared cache was enabled or disabled
 ** successfully.  An [error code] is returned otherwise.)^
 **
-** ^Shared cache is disabled by default. But this might change in
-** future releases of SQLite.  Applications that care about shared
-** cache setting should set it explicitly.
+** ^Shared cache is disabled by default. It is recommended that it stay
+** that way.  In other words, do not use this routine.  This interface
+** continues to be provided for historical compatibility, but its use is
+** discouraged.  Any use of shared cache is discouraged.  If shared cache
+** must be used, it is recommended that shared cache only be enabled for
+** individual database connections using the [sqlite3_open_v2()] interface
+** with the [SQLITE_OPEN_SHAREDCACHE] flag.
 **
 ** Note: This method is disabled on MacOS X 10.7 and iOS version 5.0
 ** and will always return SQLITE_MISUSE. On those systems, 
@@ -7077,6 +7337,9 @@ SQLITE_API int sqlite3_db_release_memory(sqlite3*);
 /*
 ** CAPI3REF: Impose A Limit On Heap Size
 **
+** These interfaces impose limits on the amount of heap memory that will be
+** by all database connections within a single process.
+**
 ** ^The sqlite3_soft_heap_limit64() interface sets and/or queries the
 ** soft limit on the amount of heap memory that may be allocated by SQLite.
 ** ^SQLite strives to keep heap memory utilization below the soft heap
@@ -7087,20 +7350,41 @@ SQLITE_API int sqlite3_db_release_memory(sqlite3*);
 ** an [SQLITE_NOMEM] error.  In other words, the soft heap limit 
 ** is advisory only.
 **
-** ^The return value from sqlite3_soft_heap_limit64() is the size of
-** the soft heap limit prior to the call, or negative in the case of an
-** error.  ^If the argument N is negative
-** then no change is made to the soft heap limit.  Hence, the current
-** size of the soft heap limit can be determined by invoking
-** sqlite3_soft_heap_limit64() with a negative argument.
-**
-** ^If the argument N is zero then the soft heap limit is disabled.
+** ^The sqlite3_hard_heap_limit64(N) interface sets a hard upper bound of
+** N bytes on the amount of memory that will be allocated.  ^The
+** sqlite3_hard_heap_limit64(N) interface is similar to
+** sqlite3_soft_heap_limit64(N) except that memory allocations will fail
+** when the hard heap limit is reached.
 **
-** ^(The soft heap limit is not enforced in the current implementation
+** ^The return value from both sqlite3_soft_heap_limit64() and
+** sqlite3_hard_heap_limit64() is the size of
+** the heap limit prior to the call, or negative in the case of an
+** error.  ^If the argument N is negative
+** then no change is made to the heap limit.  Hence, the current
+** size of heap limits can be determined by invoking
+** sqlite3_soft_heap_limit64(-1) or sqlite3_hard_heap_limit(-1).
+**
+** ^Setting the heap limits to zero disables the heap limiter mechanism.
+**
+** ^The soft heap limit may not be greater than the hard heap limit.
+** ^If the hard heap limit is enabled and if sqlite3_soft_heap_limit(N)
+** is invoked with a value of N that is greater than the hard heap limit,
+** the the soft heap limit is set to the value of the hard heap limit.
+** ^The soft heap limit is automatically enabled whenever the hard heap
+** limit is enabled. ^When sqlite3_hard_heap_limit64(N) is invoked and
+** the soft heap limit is outside the range of 1..N, then the soft heap
+** limit is set to N.  ^Invoking sqlite3_soft_heap_limit64(0) when the
+** hard heap limit is enabled makes the soft heap limit equal to the
+** hard heap limit.
+**
+** The memory allocation limits can also be adjusted using
+** [PRAGMA soft_heap_limit] and [PRAGMA hard_heap_limit].
+**
+** ^(The heap limits are not enforced in the current implementation
 ** if one or more of following conditions are true:
 **
 ** <ul>
-** <li> The soft heap limit is set to zero.
+** <li> The limit value is set to zero.
 ** <li> Memory accounting is disabled using a combination of the
 **      [sqlite3_config]([SQLITE_CONFIG_MEMSTATUS],...) start-time option and
 **      the [SQLITE_DEFAULT_MEMSTATUS] compile-time option.
@@ -7111,21 +7395,11 @@ SQLITE_API int sqlite3_db_release_memory(sqlite3*);
 **      from the heap.
 ** </ul>)^
 **
-** Beginning with SQLite [version 3.7.3] ([dateof:3.7.3]), 
-** the soft heap limit is enforced
-** regardless of whether or not the [SQLITE_ENABLE_MEMORY_MANAGEMENT]
-** compile-time option is invoked.  With [SQLITE_ENABLE_MEMORY_MANAGEMENT],
-** the soft heap limit is enforced on every memory allocation.  Without
-** [SQLITE_ENABLE_MEMORY_MANAGEMENT], the soft heap limit is only enforced
-** when memory is allocated by the page cache.  Testing suggests that because
-** the page cache is the predominate memory user in SQLite, most
-** applications will achieve adequate soft heap limit enforcement without
-** the use of [SQLITE_ENABLE_MEMORY_MANAGEMENT].
-**
-** The circumstances under which SQLite will enforce the soft heap limit may
+** The circumstances under which SQLite will enforce the heap limits may
 ** changes in future releases of SQLite.
 */
 SQLITE_API sqlite3_int64 sqlite3_soft_heap_limit64(sqlite3_int64 N);
+SQLITE_API sqlite3_int64 sqlite3_hard_heap_limit64(sqlite3_int64 N);
 
 /*
 ** CAPI3REF: Deprecated Soft Heap Limit Interface
@@ -7149,7 +7423,7 @@ SQLITE_API SQLITE_DEPRECATED void sqlite3_soft_heap_limit(int N);
 ** interface returns SQLITE_OK and fills in the non-NULL pointers in
 ** the final five arguments with appropriate values if the specified
 ** column exists.  ^The sqlite3_table_column_metadata() interface returns
-** SQLITE_ERROR and if the specified column does not exist.
+** SQLITE_ERROR if the specified column does not exist.
 ** ^If the column-name parameter to sqlite3_table_column_metadata() is a
 ** NULL pointer, then this routine simply checks for the existence of the
 ** table and returns SQLITE_OK if the table exists and SQLITE_ERROR if it
@@ -7291,7 +7565,7 @@ SQLITE_API int sqlite3_load_extension(
 ** to enable or disable only the C-API.)^
 **
 ** <b>Security warning:</b> It is recommended that extension loading
-** be disabled using the [SQLITE_DBCONFIG_ENABLE_LOAD_EXTENSION] method
+** be enabled using the [SQLITE_DBCONFIG_ENABLE_LOAD_EXTENSION] method
 ** rather than this interface, so the [load_extension()] SQL function
 ** remains disabled. This will prevent SQL injections from giving attackers
 ** access to extension loading capabilities.
@@ -7378,7 +7652,7 @@ typedef struct sqlite3_module sqlite3_module;
 ** KEYWORDS: sqlite3_module {virtual table module}
 **
 ** This structure, sometimes called a "virtual table module", 
-** defines the implementation of a [virtual tables].  
+** defines the implementation of a [virtual table].  
 ** This structure consists mostly of methods for the module.
 **
 ** ^A virtual table module is created by filling in a persistent
@@ -7475,7 +7749,13 @@ struct sqlite3_module {
 ** the right-hand side of the corresponding aConstraint[] is evaluated
 ** and becomes the argvIndex-th entry in argv.  ^(If aConstraintUsage[].omit
 ** is true, then the constraint is assumed to be fully handled by the
-** virtual table and is not checked again by SQLite.)^
+** virtual table and might not be checked again by the byte code.)^ ^(The
+** aConstraintUsage[].omit flag is an optimization hint. When the omit flag
+** is left in its default setting of false, the constraint will always be
+** checked separately in byte code.  If the omit flag is change to true, then
+** the constraint may or may not be checked in byte code.  In other words,
+** when the omit flag is true there is no guarantee that the constraint will
+** not be checked again using byte code.)^
 **
 ** ^The idxNum and idxPtr values are recorded and passed into the
 ** [xFilter] method.
@@ -7515,7 +7795,7 @@ struct sqlite3_module {
 ** If a virtual table extension is
 ** used with an SQLite version earlier than 3.8.2, the results of attempting 
 ** to read or write the estimatedRows field are undefined (but are likely 
-** to included crashing the application). The estimatedRows field should
+** to include crashing the application). The estimatedRows field should
 ** therefore only be used if [sqlite3_libversion_number()] returns a
 ** value greater than or equal to 3008002. Similarly, the idxFlags field
 ** was added for [version 3.9.0] ([dateof:3.9.0]). 
@@ -7567,7 +7847,7 @@ struct sqlite3_index_info {
 /*
 ** CAPI3REF: Virtual Table Constraint Operator Codes
 **
-** These macros defined the allowed values for the
+** These macros define the allowed values for the
 ** [sqlite3_index_info].aConstraint[].op field.  Each value represents
 ** an operator that is part of a constraint term in the wHERE clause of
 ** a query that uses a [virtual table].
@@ -7613,6 +7893,12 @@ struct sqlite3_index_info {
 ** ^The sqlite3_create_module()
 ** interface is equivalent to sqlite3_create_module_v2() with a NULL
 ** destructor.
+**
+** ^If the third parameter (the pointer to the sqlite3_module object) is
+** NULL then no new module is create and any existing modules with the
+** same name are dropped.
+**
+** See also: [sqlite3_drop_modules()]
 */
 SQLITE_API int sqlite3_create_module(
   sqlite3 *db,               /* SQLite connection to register module with */
@@ -7628,6 +7914,23 @@ SQLITE_API int sqlite3_create_module_v2(
   void(*xDestroy)(void*)     /* Module destructor function */
 );
 
+/*
+** CAPI3REF: Remove Unnecessary Virtual Table Implementations
+** METHOD: sqlite3
+**
+** ^The sqlite3_drop_modules(D,L) interface removes all virtual
+** table modules from database connection D except those named on list L.
+** The L parameter must be either NULL or a pointer to an array of pointers
+** to strings where the array is terminated by a single NULL pointer.
+** ^If the L parameter is NULL, then all virtual table modules are removed.
+**
+** See also: [sqlite3_create_module()]
+*/
+SQLITE_API int sqlite3_drop_modules(
+  sqlite3 *db,                /* Remove modules from this connection */
+  const char **azKeep         /* Except, do not remove the ones named here */
+);
+
 /*
 ** CAPI3REF: Virtual Table Instance Object
 ** KEYWORDS: sqlite3_vtab
@@ -8154,7 +8457,7 @@ SQLITE_API void sqlite3_mutex_leave(sqlite3_mutex*);
 ** The only difference is that the public sqlite3_XXX functions enumerated
 ** above silently ignore any invocations that pass a NULL pointer instead
 ** of a valid mutex handle. The implementations of the methods defined
-** by this structure are not required to handle this case, the results
+** by this structure are not required to handle this case. The results
 ** of passing a NULL pointer instead of a valid mutex handle are undefined
 ** (i.e. it is acceptable to provide an implementation that segfaults if
 ** it is passed a NULL pointer).
@@ -8336,7 +8639,7 @@ SQLITE_API int sqlite3_test_control(int op, ...);
 #define SQLITE_TESTCTRL_FIRST                    5
 #define SQLITE_TESTCTRL_PRNG_SAVE                5
 #define SQLITE_TESTCTRL_PRNG_RESTORE             6
-#define SQLITE_TESTCTRL_PRNG_RESET               7
+#define SQLITE_TESTCTRL_PRNG_RESET               7  /* NOT USED */
 #define SQLITE_TESTCTRL_BITVEC_TEST              8
 #define SQLITE_TESTCTRL_FAULT_INSTALL            9
 #define SQLITE_TESTCTRL_BENIGN_MALLOC_HOOKS     10
@@ -8358,7 +8661,10 @@ SQLITE_API int sqlite3_test_control(int op, ...);
 #define SQLITE_TESTCTRL_SORTER_MMAP             24
 #define SQLITE_TESTCTRL_IMPOSTER                25
 #define SQLITE_TESTCTRL_PARSER_COVERAGE         26
-#define SQLITE_TESTCTRL_LAST                    26  /* Largest TESTCTRL */
+#define SQLITE_TESTCTRL_RESULT_INTREAL          27
+#define SQLITE_TESTCTRL_PRNG_SEED               28
+#define SQLITE_TESTCTRL_EXTRA_SCHEMA_CHECKS     29
+#define SQLITE_TESTCTRL_LAST                    29  /* Largest TESTCTRL */
 
 /*
 ** CAPI3REF: SQL Keyword Checking
@@ -8624,7 +8930,7 @@ SQLITE_API int sqlite3_status64(
 **
 ** [[SQLITE_STATUS_PAGECACHE_SIZE]] ^(<dt>SQLITE_STATUS_PAGECACHE_SIZE</dt>
 ** <dd>This parameter records the largest memory allocation request
-** handed to [pagecache memory allocator].  Only the value returned in the
+** handed to the [pagecache memory allocator].  Only the value returned in the
 ** *pHighwater parameter to [sqlite3_status()] is of interest.  
 ** The value written into the *pCurrent parameter is undefined.</dd>)^
 **
@@ -8700,7 +9006,7 @@ SQLITE_API int sqlite3_db_status(sqlite3*, int op, int *pCur, int *pHiwtr, int r
 ** checked out.</dd>)^
 **
 ** [[SQLITE_DBSTATUS_LOOKASIDE_HIT]] ^(<dt>SQLITE_DBSTATUS_LOOKASIDE_HIT</dt>
-** <dd>This parameter returns the number malloc attempts that were 
+** <dd>This parameter returns the number of malloc attempts that were 
 ** satisfied using lookaside memory. Only the high-water value is meaningful;
 ** the current value is always zero.)^
 **
@@ -8782,7 +9088,7 @@ SQLITE_API int sqlite3_db_status(sqlite3*, int op, int *pCur, int *pHiwtr, int r
 ** cache overflowing. Transactions are more efficient if they are written
 ** to disk all at once. When pages spill mid-transaction, that introduces
 ** additional overhead. This parameter can be used help identify
-** inefficiencies that can be resolve by increasing the cache size.
+** inefficiencies that can be resolved by increasing the cache size.
 ** </dd>
 **
 ** [[SQLITE_DBSTATUS_DEFERRED_FKS]] ^(<dt>SQLITE_DBSTATUS_DEFERRED_FKS</dt>
@@ -8871,7 +9177,7 @@ SQLITE_API int sqlite3_stmt_status(sqlite3_stmt*, int op,int resetFlg);
 **
 ** [[SQLITE_STMTSTATUS_REPREPARE]] <dt>SQLITE_STMTSTATUS_REPREPARE</dt>
 ** <dd>^This is the number of times that the prepare statement has been
-** automatically regenerated due to schema changes or change to 
+** automatically regenerated due to schema changes or changes to 
 ** [bound parameters] that might affect the query plan.
 **
 ** [[SQLITE_STMTSTATUS_RUN]] <dt>SQLITE_STMTSTATUS_RUN</dt>
@@ -9042,7 +9348,7 @@ struct sqlite3_pcache_page {
 **
 ** ^(SQLite will normally invoke xFetch() with a createFlag of 0 or 1.  SQLite
 ** will only use a createFlag of 2 after a prior call with a createFlag of 1
-** failed.)^  In between the to xFetch() calls, SQLite may
+** failed.)^  In between the xFetch() calls, SQLite may
 ** attempt to unpin one or more cache pages by spilling the content of
 ** pinned pages to disk and synching the operating system disk cache.
 **
@@ -9360,7 +9666,7 @@ SQLITE_API int sqlite3_backup_pagecount(sqlite3_backup *p);
 ** the first argument to register for a callback that will be invoked
 ** when the blocking connections current transaction is concluded. ^The
 ** callback is invoked from within the [sqlite3_step] or [sqlite3_close]
-** call that concludes the blocking connections transaction.
+** call that concludes the blocking connection's transaction.
 **
 ** ^(If sqlite3_unlock_notify() is called in a multi-threaded application,
 ** there is a chance that the blocking connection will have already
@@ -9398,7 +9704,7 @@ SQLITE_API int sqlite3_backup_pagecount(sqlite3_backup *p);
 ** an unlock-notify callback is a pointer to an array of void* pointers,
 ** and the second is the number of entries in the array.
 **
-** When a blocking connections transaction is concluded, there may be
+** When a blocking connection's transaction is concluded, there may be
 ** more than one blocked connection that has registered for an unlock-notify
 ** callback. ^If two or more such blocked connections have specified the
 ** same callback function, then instead of invoking the callback function
@@ -9746,14 +10052,20 @@ SQLITE_API int sqlite3_wal_checkpoint_v2(
 ** If this interface is invoked outside the context of an xConnect or
 ** xCreate virtual table method then the behavior is undefined.
 **
-** At present, there is only one option that may be configured using
-** this function. (See [SQLITE_VTAB_CONSTRAINT_SUPPORT].)  Further options
-** may be added in the future.
+** In the call sqlite3_vtab_config(D,C,...) the D parameter is the
+** [database connection] in which the virtual table is being created and
+** which is passed in as the first argument to the [xConnect] or [xCreate]
+** method that is invoking sqlite3_vtab_config().  The C parameter is one
+** of the [virtual table configuration options].  The presence and meaning
+** of parameters after C depend on which [virtual table configuration option]
+** is used.
 */
 SQLITE_API int sqlite3_vtab_config(sqlite3*, int op, ...);
 
 /*
 ** CAPI3REF: Virtual Table Configuration Options
+** KEYWORDS: {virtual table configuration options} 
+** KEYWORDS: {virtual table configuration option}
 **
 ** These macros define the various options to the
 ** [sqlite3_vtab_config()] interface that [virtual table] implementations
@@ -9761,7 +10073,7 @@ SQLITE_API int sqlite3_vtab_config(sqlite3*, int op, ...);
 **
 ** <dl>
 ** [[SQLITE_VTAB_CONSTRAINT_SUPPORT]]
-** <dt>SQLITE_VTAB_CONSTRAINT_SUPPORT
+** <dt>SQLITE_VTAB_CONSTRAINT_SUPPORT</dt>
 ** <dd>Calls of the form
 ** [sqlite3_vtab_config](db,SQLITE_VTAB_CONSTRAINT_SUPPORT,X) are supported,
 ** where X is an integer.  If X is zero, then the [virtual table] whose
@@ -9790,9 +10102,31 @@ SQLITE_API int sqlite3_vtab_config(sqlite3*, int op, ...);
 ** return SQLITE_OK. Or, if this is not possible, it may return
 ** SQLITE_CONSTRAINT, in which case SQLite falls back to OR ABORT 
 ** constraint handling.
+** </dd>
+**
+** [[SQLITE_VTAB_DIRECTONLY]]<dt>SQLITE_VTAB_DIRECTONLY</dt>
+** <dd>Calls of the form
+** [sqlite3_vtab_config](db,SQLITE_VTAB_DIRECTONLY) from within the
+** the [xConnect] or [xCreate] methods of a [virtual table] implmentation
+** prohibits that virtual table from being used from within triggers and
+** views.
+** </dd>
+**
+** [[SQLITE_VTAB_INNOCUOUS]]<dt>SQLITE_VTAB_INNOCUOUS</dt>
+** <dd>Calls of the form
+** [sqlite3_vtab_config](db,SQLITE_VTAB_INNOCUOUS) from within the
+** the [xConnect] or [xCreate] methods of a [virtual table] implmentation
+** identify that virtual table as being safe to use from within triggers
+** and views.  Conceptually, the SQLITE_VTAB_INNOCUOUS tag means that the
+** virtual table can do no serious harm even if it is controlled by a
+** malicious hacker.  Developers should avoid setting the SQLITE_VTAB_INNOCUOUS
+** flag unless absolutely necessary.
+** </dd>
 ** </dl>
 */
 #define SQLITE_VTAB_CONSTRAINT_SUPPORT 1
+#define SQLITE_VTAB_INNOCUOUS          2
+#define SQLITE_VTAB_DIRECTONLY         3
 
 /*
 ** CAPI3REF: Determine The Virtual Table Conflict Policy
@@ -9872,15 +10206,15 @@ SQLITE_API SQLITE_EXPERIMENTAL const char *sqlite3_vtab_collation(sqlite3_index_
 **
 ** <dl>
 ** [[SQLITE_SCANSTAT_NLOOP]] <dt>SQLITE_SCANSTAT_NLOOP</dt>
-** <dd>^The [sqlite3_int64] variable pointed to by the T parameter will be
+** <dd>^The [sqlite3_int64] variable pointed to by the V parameter will be
 ** set to the total number of times that the X-th loop has run.</dd>
 **
 ** [[SQLITE_SCANSTAT_NVISIT]] <dt>SQLITE_SCANSTAT_NVISIT</dt>
-** <dd>^The [sqlite3_int64] variable pointed to by the T parameter will be set
+** <dd>^The [sqlite3_int64] variable pointed to by the V parameter will be set
 ** to the total number of rows examined by all iterations of the X-th loop.</dd>
 **
 ** [[SQLITE_SCANSTAT_EST]] <dt>SQLITE_SCANSTAT_EST</dt>
-** <dd>^The "double" variable pointed to by the T parameter will be set to the
+** <dd>^The "double" variable pointed to by the V parameter will be set to the
 ** query planner's estimate for the average number of rows output from each
 ** iteration of the X-th loop.  If the query planner's estimates was accurate,
 ** then this value will approximate the quotient NVISIT/NLOOP and the
@@ -9888,17 +10222,17 @@ SQLITE_API SQLITE_EXPERIMENTAL const char *sqlite3_vtab_collation(sqlite3_index_
 ** be the NLOOP value for the current loop.
 **
 ** [[SQLITE_SCANSTAT_NAME]] <dt>SQLITE_SCANSTAT_NAME</dt>
-** <dd>^The "const char *" variable pointed to by the T parameter will be set
+** <dd>^The "const char *" variable pointed to by the V parameter will be set
 ** to a zero-terminated UTF-8 string containing the name of the index or table
 ** used for the X-th loop.
 **
 ** [[SQLITE_SCANSTAT_EXPLAIN]] <dt>SQLITE_SCANSTAT_EXPLAIN</dt>
-** <dd>^The "const char *" variable pointed to by the T parameter will be set
+** <dd>^The "const char *" variable pointed to by the V parameter will be set
 ** to a zero-terminated UTF-8 string containing the [EXPLAIN QUERY PLAN]
 ** description for the X-th loop.
 **
 ** [[SQLITE_SCANSTAT_SELECTID]] <dt>SQLITE_SCANSTAT_SELECT</dt>
-** <dd>^The "int" variable pointed to by the T parameter will be set to the
+** <dd>^The "int" variable pointed to by the V parameter will be set to the
 ** "select-id" for the X-th loop.  The select-id identifies which query or
 ** subquery the loop is part of.  The main query has a select-id of zero.
 ** The select-id is the same value as is output in the first column
@@ -10753,7 +11087,7 @@ SQLITE_API int sqlite3session_attach(
 ** The second argument (xFilter) is the "filter callback". For changes to rows 
 ** in tables that are not attached to the Session object, the filter is called
 ** to determine whether changes to the table's rows should be tracked or not. 
-** If xFilter returns 0, changes is not tracked. Note that once a table is 
+** If xFilter returns 0, changes are not tracked. Note that once a table is 
 ** attached, xFilter will not be called again.
 */
 SQLITE_API void sqlite3session_table_filter(
@@ -10927,7 +11261,7 @@ SQLITE_API int sqlite3session_changeset(
 ** It an error if database zFrom does not exist or does not contain the
 ** required compatible table.
 **
-** If the operation successful, SQLITE_OK is returned. Otherwise, an SQLite
+** If the operation is successful, SQLITE_OK is returned. Otherwise, an SQLite
 ** error code. In this case, if argument pzErrMsg is not NULL, *pzErrMsg
 ** may be set to point to a buffer containing an English language error 
 ** message. It is the responsibility of the caller to free this buffer using
@@ -11064,7 +11398,7 @@ SQLITE_API int sqlite3changeset_start_v2(
 ** CAPI3REF: Advance A Changeset Iterator
 ** METHOD: sqlite3_changeset_iter
 **
-** This function may only be used with iterators created by function
+** This function may only be used with iterators created by the function
 ** [sqlite3changeset_start()]. If it is called on an iterator passed to
 ** a conflict-handler callback by [sqlite3changeset_apply()], SQLITE_MISUSE
 ** is returned and the call has no effect.
@@ -11480,8 +11814,8 @@ SQLITE_API int sqlite3changegroup_new(sqlite3_changegroup **pp);
 ** case, this function fails with SQLITE_SCHEMA. If the input changeset
 ** appears to be corrupt and the corruption is detected, SQLITE_CORRUPT is
 ** returned. Or, if an out-of-memory condition occurs during processing, this
-** function returns SQLITE_NOMEM. In all cases, if an error occurs the
-** final contents of the changegroup is undefined.
+** function returns SQLITE_NOMEM. In all cases, if an error occurs the state
+** of the final contents of the changegroup is undefined.
 **
 ** If no error occurs, SQLITE_OK is returned.
 */
@@ -11656,7 +11990,7 @@ SQLITE_API void sqlite3changegroup_delete(sqlite3_changegroup*);
 **
 ** It is safe to execute SQL statements, including those that write to the
 ** table that the callback related to, from within the xConflict callback.
-** This can be used to further customize the applications conflict
+** This can be used to further customize the application's conflict
 ** resolution strategy.
 **
 ** All changes made by these functions are enclosed in a savepoint transaction.
@@ -11966,7 +12300,7 @@ SQLITE_API int sqlite3rebaser_configure(
 **
 ** Argument pIn must point to a buffer containing a changeset nIn bytes
 ** in size. This function allocates and populates a buffer with a copy
-** of the changeset rebased rebased according to the configuration of the
+** of the changeset rebased according to the configuration of the
 ** rebaser object passed as the first argument. If successful, (*ppOut)
 ** is set to point to the new buffer containing the rebased changeset and 
 ** (*pnOut) to its size in bytes and SQLITE_OK returned. It is the
@@ -12374,7 +12708,7 @@ struct Fts5PhraseIter {
 **
 ** xSetAuxdata(pFts5, pAux, xDelete)
 **
-**   Save the pointer passed as the second argument as the extension functions 
+**   Save the pointer passed as the second argument as the extension function's 
 **   "auxiliary data". The pointer may then be retrieved by the current or any
 **   future invocation of the same fts5 extension function made as part of
 **   the same MATCH query using the xGetAuxdata() API.
@@ -12616,8 +12950,8 @@ struct Fts5ExtensionApi {
 **
 **   There are several ways to approach this in FTS5:
 **
-**   <ol><li> By mapping all synonyms to a single token. In this case, the 
-**            In the above example, this means that the tokenizer returns the
+**   <ol><li> By mapping all synonyms to a single token. In this case, using
+**            the above example, this means that the tokenizer returns the
 **            same token for inputs "first" and "1st". Say that token is in
 **            fact "first", so that when the user inserts the document "I won
 **            1st place" entries are added to the index for tokens "i", "won",
@@ -13055,15 +13389,15 @@ struct fts5_api {
 ** So we have to define the macros in different ways depending on the
 ** compiler.
 */
-#if defined(__PTRDIFF_TYPE__)  /* This case should work for GCC */
+#if defined(HAVE_STDINT_H)   /* Use this case if we have ANSI headers */
+# define SQLITE_INT_TO_PTR(X)  ((void*)(intptr_t)(X))
+# define SQLITE_PTR_TO_INT(X)  ((int)(intptr_t)(X))
+#elif defined(__PTRDIFF_TYPE__)  /* This case should work for GCC */
 # define SQLITE_INT_TO_PTR(X)  ((void*)(__PTRDIFF_TYPE__)(X))
 # define SQLITE_PTR_TO_INT(X)  ((int)(__PTRDIFF_TYPE__)(X))
 #elif !defined(__GNUC__)       /* Works for compilers other than LLVM */
 # define SQLITE_INT_TO_PTR(X)  ((void*)&((char*)0)[X])
 # define SQLITE_PTR_TO_INT(X)  ((int)(((char*)X)-(char*)0))
-#elif defined(HAVE_STDINT_H)   /* Use this case if we have ANSI headers */
-# define SQLITE_INT_TO_PTR(X)  ((void*)(intptr_t)(X))
-# define SQLITE_PTR_TO_INT(X)  ((int)(intptr_t)(X))
 #else                          /* Generates a warning - but it always works */
 # define SQLITE_INT_TO_PTR(X)  ((void*)(X))
 # define SQLITE_PTR_TO_INT(X)  ((int)(X))
@@ -13289,6 +13623,26 @@ SQLITE_PRIVATE   void sqlite3Coverage(int);
 # define NEVER(X)       (X)
 #endif
 
+/*
+** The harmless(X) macro indicates that expression X is usually false
+** but can be true without causing any problems, but we don't know of
+** any way to cause X to be true.
+**
+** In debugging and testing builds, this macro will abort if X is ever
+** true.  In this way, developers are alerted to a possible test case
+** that causes X to be true.  If a harmless macro ever fails, that is
+** an opportunity to change the macro into a testcase() and add a new
+** test case to the test suite.
+**
+** For normal production builds, harmless(X) is a no-op, since it does
+** not matter whether expression X is true or false.
+*/
+#ifdef SQLITE_DEBUG
+# define harmless(X)  assert(!(X));
+#else
+# define harmless(X)
+#endif
+
 /*
 ** Some conditionals are optimizations only.  In other words, if the
 ** conditionals are replaced with a constant 1 (true) or 0 (false) then
@@ -13553,100 +13907,105 @@ SQLITE_PRIVATE void sqlite3HashClear(Hash*);
 #define TK_VIEW                            79
 #define TK_VIRTUAL                         80
 #define TK_WITH                            81
-#define TK_CURRENT                         82
-#define TK_FOLLOWING                       83
-#define TK_PARTITION                       84
-#define TK_PRECEDING                       85
-#define TK_RANGE                           86
-#define TK_UNBOUNDED                       87
-#define TK_EXCLUDE                         88
-#define TK_GROUPS                          89
-#define TK_OTHERS                          90
-#define TK_TIES                            91
-#define TK_REINDEX                         92
-#define TK_RENAME                          93
-#define TK_CTIME_KW                        94
-#define TK_ANY                             95
-#define TK_BITAND                          96
-#define TK_BITOR                           97
-#define TK_LSHIFT                          98
-#define TK_RSHIFT                          99
-#define TK_PLUS                           100
-#define TK_MINUS                          101
-#define TK_STAR                           102
-#define TK_SLASH                          103
-#define TK_REM                            104
-#define TK_CONCAT                         105
-#define TK_COLLATE                        106
-#define TK_BITNOT                         107
-#define TK_ON                             108
-#define TK_INDEXED                        109
-#define TK_STRING                         110
-#define TK_JOIN_KW                        111
-#define TK_CONSTRAINT                     112
-#define TK_DEFAULT                        113
-#define TK_NULL                           114
-#define TK_PRIMARY                        115
-#define TK_UNIQUE                         116
-#define TK_CHECK                          117
-#define TK_REFERENCES                     118
-#define TK_AUTOINCR                       119
-#define TK_INSERT                         120
-#define TK_DELETE                         121
-#define TK_UPDATE                         122
-#define TK_SET                            123
-#define TK_DEFERRABLE                     124
-#define TK_FOREIGN                        125
-#define TK_DROP                           126
-#define TK_UNION                          127
-#define TK_ALL                            128
-#define TK_EXCEPT                         129
-#define TK_INTERSECT                      130
-#define TK_SELECT                         131
-#define TK_VALUES                         132
-#define TK_DISTINCT                       133
-#define TK_DOT                            134
-#define TK_FROM                           135
-#define TK_JOIN                           136
-#define TK_USING                          137
-#define TK_ORDER                          138
-#define TK_GROUP                          139
-#define TK_HAVING                         140
-#define TK_LIMIT                          141
-#define TK_WHERE                          142
-#define TK_INTO                           143
-#define TK_NOTHING                        144
-#define TK_FLOAT                          145
-#define TK_BLOB                           146
-#define TK_INTEGER                        147
-#define TK_VARIABLE                       148
-#define TK_CASE                           149
-#define TK_WHEN                           150
-#define TK_THEN                           151
-#define TK_ELSE                           152
-#define TK_INDEX                          153
-#define TK_ALTER                          154
-#define TK_ADD                            155
-#define TK_WINDOW                         156
-#define TK_OVER                           157
-#define TK_FILTER                         158
-#define TK_TRUEFALSE                      159
-#define TK_ISNOT                          160
-#define TK_FUNCTION                       161
-#define TK_COLUMN                         162
-#define TK_AGG_FUNCTION                   163
-#define TK_AGG_COLUMN                     164
-#define TK_UMINUS                         165
-#define TK_UPLUS                          166
-#define TK_TRUTH                          167
-#define TK_REGISTER                       168
-#define TK_VECTOR                         169
-#define TK_SELECT_COLUMN                  170
-#define TK_IF_NULL_ROW                    171
-#define TK_ASTERISK                       172
-#define TK_SPAN                           173
-#define TK_SPACE                          174
-#define TK_ILLEGAL                        175
+#define TK_NULLS                           82
+#define TK_FIRST                           83
+#define TK_LAST                            84
+#define TK_CURRENT                         85
+#define TK_FOLLOWING                       86
+#define TK_PARTITION                       87
+#define TK_PRECEDING                       88
+#define TK_RANGE                           89
+#define TK_UNBOUNDED                       90
+#define TK_EXCLUDE                         91
+#define TK_GROUPS                          92
+#define TK_OTHERS                          93
+#define TK_TIES                            94
+#define TK_GENERATED                       95
+#define TK_ALWAYS                          96
+#define TK_REINDEX                         97
+#define TK_RENAME                          98
+#define TK_CTIME_KW                        99
+#define TK_ANY                            100
+#define TK_BITAND                         101
+#define TK_BITOR                          102
+#define TK_LSHIFT                         103
+#define TK_RSHIFT                         104
+#define TK_PLUS                           105
+#define TK_MINUS                          106
+#define TK_STAR                           107
+#define TK_SLASH                          108
+#define TK_REM                            109
+#define TK_CONCAT                         110
+#define TK_COLLATE                        111
+#define TK_BITNOT                         112
+#define TK_ON                             113
+#define TK_INDEXED                        114
+#define TK_STRING                         115
+#define TK_JOIN_KW                        116
+#define TK_CONSTRAINT                     117
+#define TK_DEFAULT                        118
+#define TK_NULL                           119
+#define TK_PRIMARY                        120
+#define TK_UNIQUE                         121
+#define TK_CHECK                          122
+#define TK_REFERENCES                     123
+#define TK_AUTOINCR                       124
+#define TK_INSERT                         125
+#define TK_DELETE                         126
+#define TK_UPDATE                         127
+#define TK_SET                            128
+#define TK_DEFERRABLE                     129
+#define TK_FOREIGN                        130
+#define TK_DROP                           131
+#define TK_UNION                          132
+#define TK_ALL                            133
+#define TK_EXCEPT                         134
+#define TK_INTERSECT                      135
+#define TK_SELECT                         136
+#define TK_VALUES                         137
+#define TK_DISTINCT                       138
+#define TK_DOT                            139
+#define TK_FROM                           140
+#define TK_JOIN                           141
+#define TK_USING                          142
+#define TK_ORDER                          143
+#define TK_GROUP                          144
+#define TK_HAVING                         145
+#define TK_LIMIT                          146
+#define TK_WHERE                          147
+#define TK_INTO                           148
+#define TK_NOTHING                        149
+#define TK_FLOAT                          150
+#define TK_BLOB                           151
+#define TK_INTEGER                        152
+#define TK_VARIABLE                       153
+#define TK_CASE                           154
+#define TK_WHEN                           155
+#define TK_THEN                           156
+#define TK_ELSE                           157
+#define TK_INDEX                          158
+#define TK_ALTER                          159
+#define TK_ADD                            160
+#define TK_WINDOW                         161
+#define TK_OVER                           162
+#define TK_FILTER                         163
+#define TK_COLUMN                         164
+#define TK_AGG_FUNCTION                   165
+#define TK_AGG_COLUMN                     166
+#define TK_TRUEFALSE                      167
+#define TK_ISNOT                          168
+#define TK_FUNCTION                       169
+#define TK_UMINUS                         170
+#define TK_UPLUS                          171
+#define TK_TRUTH                          172
+#define TK_REGISTER                       173
+#define TK_VECTOR                         174
+#define TK_SELECT_COLUMN                  175
+#define TK_IF_NULL_ROW                    176
+#define TK_ASTERISK                       177
+#define TK_SPAN                           178
+#define TK_SPACE                          179
+#define TK_ILLEGAL                        180
 
 /************** End of parse.h ***********************************************/
 /************** Continuing where we left off in sqliteInt.h ******************/
@@ -13952,12 +14311,13 @@ typedef INT16_TYPE LogEst;
 ** at run-time.
 */
 #ifndef SQLITE_BYTEORDER
-# if defined(i386)     || defined(__i386__)   || defined(_M_IX86) ||    \
-     defined(__x86_64) || defined(__x86_64__) || defined(_M_X64)  ||    \
-     defined(_M_AMD64) || defined(_M_ARM)     || defined(__x86)   ||    \
-     defined(__arm__)  || defined(_M_ARM64)
+# if defined(i386)      || defined(__i386__)      || defined(_M_IX86) ||    \
+     defined(__x86_64)  || defined(__x86_64__)    || defined(_M_X64)  ||    \
+     defined(_M_AMD64)  || defined(_M_ARM)        || defined(__x86)   ||    \
+     defined(__ARMEL__) || defined(__AARCH64EL__) || defined(_M_ARM64)
 #   define SQLITE_BYTEORDER    1234
-# elif defined(sparc)    || defined(__ppc__)
+# elif defined(sparc)     || defined(__ppc__) || \
+       defined(__ARMEB__) || defined(__AARCH64EB__)
 #   define SQLITE_BYTEORDER    4321
 # else
 #   define SQLITE_BYTEORDER 0
@@ -14056,20 +14416,6 @@ typedef INT16_TYPE LogEst;
 # define SQLITE_DEFAULT_MMAP_SIZE SQLITE_MAX_MMAP_SIZE
 #endif
 
-/*
-** Only one of SQLITE_ENABLE_STAT3 or SQLITE_ENABLE_STAT4 can be defined.
-** Priority is given to SQLITE_ENABLE_STAT4.  If either are defined, also
-** define SQLITE_ENABLE_STAT3_OR_STAT4
-*/
-#ifdef SQLITE_ENABLE_STAT4
-# undef SQLITE_ENABLE_STAT3
-# define SQLITE_ENABLE_STAT3_OR_STAT4 1
-#elif SQLITE_ENABLE_STAT3
-# define SQLITE_ENABLE_STAT3_OR_STAT4 1
-#elif SQLITE_ENABLE_STAT3_OR_STAT4
-# undef SQLITE_ENABLE_STAT3_OR_STAT4
-#endif
-
 /*
 ** SELECTTRACE_ENABLED will be either 1 or 0 depending on whether or not
 ** the Select query generator tracing logic is turned on.
@@ -14257,6 +14603,7 @@ typedef struct With With;
 ** A bit in a Bitmask
 */
 #define MASKBIT(n)   (((Bitmask)1)<<(n))
+#define MASKBIT64(n) (((u64)1)<<(n))
 #define MASKBIT32(n) (((unsigned int)1)<<(n))
 #define ALLBITS      ((Bitmask)-1)
 
@@ -14583,6 +14930,8 @@ SQLITE_PRIVATE int sqlite3BtreeNext(BtCursor*, int flags);
 SQLITE_PRIVATE int sqlite3BtreeEof(BtCursor*);
 SQLITE_PRIVATE int sqlite3BtreePrevious(BtCursor*, int flags);
 SQLITE_PRIVATE i64 sqlite3BtreeIntegerKey(BtCursor*);
+SQLITE_PRIVATE void sqlite3BtreeCursorPin(BtCursor*);
+SQLITE_PRIVATE void sqlite3BtreeCursorUnpin(BtCursor*);
 #ifdef SQLITE_ENABLE_OFFSET_SQL_FUNC
 SQLITE_PRIVATE i64 sqlite3BtreeOffset(BtCursor*);
 #endif
@@ -14591,7 +14940,7 @@ SQLITE_PRIVATE const void *sqlite3BtreePayloadFetch(BtCursor*, u32 *pAmt);
 SQLITE_PRIVATE u32 sqlite3BtreePayloadSize(BtCursor*);
 SQLITE_PRIVATE sqlite3_int64 sqlite3BtreeMaxRecordSize(BtCursor*);
 
-SQLITE_PRIVATE char *sqlite3BtreeIntegrityCheck(Btree*, int *aRoot, int nRoot, int, int*);
+SQLITE_PRIVATE char *sqlite3BtreeIntegrityCheck(sqlite3*,Btree*,int*aRoot,int nRoot,int,int*);
 SQLITE_PRIVATE struct Pager *sqlite3BtreePager(Btree*);
 SQLITE_PRIVATE i64 sqlite3BtreeRowCountEst(BtCursor*);
 
@@ -14612,7 +14961,7 @@ SQLITE_PRIVATE int sqlite3BtreeCursorIsValid(BtCursor*);
 SQLITE_PRIVATE int sqlite3BtreeCursorIsValidNN(BtCursor*);
 
 #ifndef SQLITE_OMIT_BTREECOUNT
-SQLITE_PRIVATE int sqlite3BtreeCount(BtCursor *, i64 *);
+SQLITE_PRIVATE int sqlite3BtreeCount(sqlite3*, BtCursor*, i64*);
 #endif
 
 #ifdef SQLITE_TEST
@@ -14869,30 +15218,30 @@ typedef struct VdbeOpList VdbeOpList;
 #define OP_SeekLE         23 /* jump, synopsis: key=r[P3@P4]               */
 #define OP_SeekGE         24 /* jump, synopsis: key=r[P3@P4]               */
 #define OP_SeekGT         25 /* jump, synopsis: key=r[P3@P4]               */
-#define OP_IfNoHope       26 /* jump, synopsis: key=r[P3@P4]               */
-#define OP_NoConflict     27 /* jump, synopsis: key=r[P3@P4]               */
-#define OP_NotFound       28 /* jump, synopsis: key=r[P3@P4]               */
-#define OP_Found          29 /* jump, synopsis: key=r[P3@P4]               */
-#define OP_SeekRowid      30 /* jump, synopsis: intkey=r[P3]               */
-#define OP_NotExists      31 /* jump, synopsis: intkey=r[P3]               */
-#define OP_Last           32 /* jump                                       */
-#define OP_IfSmaller      33 /* jump                                       */
-#define OP_SorterSort     34 /* jump                                       */
-#define OP_Sort           35 /* jump                                       */
-#define OP_Rewind         36 /* jump                                       */
-#define OP_IdxLE          37 /* jump, synopsis: key=r[P3@P4]               */
-#define OP_IdxGT          38 /* jump, synopsis: key=r[P3@P4]               */
-#define OP_IdxLT          39 /* jump, synopsis: key=r[P3@P4]               */
-#define OP_IdxGE          40 /* jump, synopsis: key=r[P3@P4]               */
-#define OP_RowSetRead     41 /* jump, synopsis: r[P3]=rowset(P1)           */
-#define OP_RowSetTest     42 /* jump, synopsis: if r[P3] in rowset(P1) goto P2 */
+#define OP_IfNotOpen      26 /* jump, synopsis: if( !csr[P1] ) goto P2     */
+#define OP_IfNoHope       27 /* jump, synopsis: key=r[P3@P4]               */
+#define OP_NoConflict     28 /* jump, synopsis: key=r[P3@P4]               */
+#define OP_NotFound       29 /* jump, synopsis: key=r[P3@P4]               */
+#define OP_Found          30 /* jump, synopsis: key=r[P3@P4]               */
+#define OP_SeekRowid      31 /* jump, synopsis: intkey=r[P3]               */
+#define OP_NotExists      32 /* jump, synopsis: intkey=r[P3]               */
+#define OP_Last           33 /* jump                                       */
+#define OP_IfSmaller      34 /* jump                                       */
+#define OP_SorterSort     35 /* jump                                       */
+#define OP_Sort           36 /* jump                                       */
+#define OP_Rewind         37 /* jump                                       */
+#define OP_IdxLE          38 /* jump, synopsis: key=r[P3@P4]               */
+#define OP_IdxGT          39 /* jump, synopsis: key=r[P3@P4]               */
+#define OP_IdxLT          40 /* jump, synopsis: key=r[P3@P4]               */
+#define OP_IdxGE          41 /* jump, synopsis: key=r[P3@P4]               */
+#define OP_RowSetRead     42 /* jump, synopsis: r[P3]=rowset(P1)           */
 #define OP_Or             43 /* same as TK_OR, synopsis: r[P3]=(r[P1] || r[P2]) */
 #define OP_And            44 /* same as TK_AND, synopsis: r[P3]=(r[P1] && r[P2]) */
-#define OP_Program        45 /* jump                                       */
-#define OP_FkIfZero       46 /* jump, synopsis: if fkctr[P1]==0 goto P2    */
-#define OP_IfPos          47 /* jump, synopsis: if r[P1]>0 then r[P1]-=P3, goto P2 */
-#define OP_IfNotZero      48 /* jump, synopsis: if r[P1]!=0 then r[P1]--, goto P2 */
-#define OP_DecrJumpZero   49 /* jump, synopsis: if (--r[P1])==0 goto P2    */
+#define OP_RowSetTest     45 /* jump, synopsis: if r[P3] in rowset(P1) goto P2 */
+#define OP_Program        46 /* jump                                       */
+#define OP_FkIfZero       47 /* jump, synopsis: if fkctr[P1]==0 goto P2    */
+#define OP_IfPos          48 /* jump, synopsis: if r[P1]>0 then r[P1]-=P3, goto P2 */
+#define OP_IfNotZero      49 /* jump, synopsis: if r[P1]!=0 then r[P1]--, goto P2 */
 #define OP_IsNull         50 /* jump, same as TK_ISNULL, synopsis: if r[P1]==NULL goto P2 */
 #define OP_NotNull        51 /* jump, same as TK_NOTNULL, synopsis: if r[P1]!=NULL goto P2 */
 #define OP_Ne             52 /* jump, same as TK_NE, synopsis: IF r[P3]!=r[P1] */
@@ -14902,83 +15251,83 @@ typedef struct VdbeOpList VdbeOpList;
 #define OP_Lt             56 /* jump, same as TK_LT, synopsis: IF r[P3]<r[P1] */
 #define OP_Ge             57 /* jump, same as TK_GE, synopsis: IF r[P3]>=r[P1] */
 #define OP_ElseNotEq      58 /* jump, same as TK_ESCAPE                    */
-#define OP_IncrVacuum     59 /* jump                                       */
-#define OP_VNext          60 /* jump                                       */
-#define OP_Init           61 /* jump, synopsis: Start at P2                */
-#define OP_PureFunc0      62
-#define OP_Function0      63 /* synopsis: r[P3]=func(r[P2@P5])             */
-#define OP_PureFunc       64
-#define OP_Function       65 /* synopsis: r[P3]=func(r[P2@P5])             */
-#define OP_Return         66
-#define OP_EndCoroutine   67
-#define OP_HaltIfNull     68 /* synopsis: if r[P3]=null halt               */
-#define OP_Halt           69
-#define OP_Integer        70 /* synopsis: r[P2]=P1                         */
-#define OP_Int64          71 /* synopsis: r[P2]=P4                         */
-#define OP_String         72 /* synopsis: r[P2]='P4' (len=P1)              */
-#define OP_Null           73 /* synopsis: r[P2..P3]=NULL                   */
-#define OP_SoftNull       74 /* synopsis: r[P1]=NULL                       */
-#define OP_Blob           75 /* synopsis: r[P2]=P4 (len=P1)                */
-#define OP_Variable       76 /* synopsis: r[P2]=parameter(P1,P4)           */
-#define OP_Move           77 /* synopsis: r[P2@P3]=r[P1@P3]                */
-#define OP_Copy           78 /* synopsis: r[P2@P3+1]=r[P1@P3+1]            */
-#define OP_SCopy          79 /* synopsis: r[P2]=r[P1]                      */
-#define OP_IntCopy        80 /* synopsis: r[P2]=r[P1]                      */
-#define OP_ResultRow      81 /* synopsis: output=r[P1@P2]                  */
-#define OP_CollSeq        82
-#define OP_AddImm         83 /* synopsis: r[P1]=r[P1]+P2                   */
-#define OP_RealAffinity   84
-#define OP_Cast           85 /* synopsis: affinity(r[P1])                  */
-#define OP_Permutation    86
-#define OP_Compare        87 /* synopsis: r[P1@P3] <-> r[P2@P3]            */
-#define OP_IsTrue         88 /* synopsis: r[P2] = coalesce(r[P1]==TRUE,P3) ^ P4 */
-#define OP_Offset         89 /* synopsis: r[P3] = sqlite_offset(P1)        */
-#define OP_Column         90 /* synopsis: r[P3]=PX                         */
-#define OP_Affinity       91 /* synopsis: affinity(r[P1@P2])               */
-#define OP_MakeRecord     92 /* synopsis: r[P3]=mkrec(r[P1@P2])            */
-#define OP_Count          93 /* synopsis: r[P2]=count()                    */
-#define OP_ReadCookie     94
-#define OP_SetCookie      95
-#define OP_BitAnd         96 /* same as TK_BITAND, synopsis: r[P3]=r[P1]&r[P2] */
-#define OP_BitOr          97 /* same as TK_BITOR, synopsis: r[P3]=r[P1]|r[P2] */
-#define OP_ShiftLeft      98 /* same as TK_LSHIFT, synopsis: r[P3]=r[P2]<<r[P1] */
-#define OP_ShiftRight     99 /* same as TK_RSHIFT, synopsis: r[P3]=r[P2]>>r[P1] */
-#define OP_Add           100 /* same as TK_PLUS, synopsis: r[P3]=r[P1]+r[P2] */
-#define OP_Subtract      101 /* same as TK_MINUS, synopsis: r[P3]=r[P2]-r[P1] */
-#define OP_Multiply      102 /* same as TK_STAR, synopsis: r[P3]=r[P1]*r[P2] */
-#define OP_Divide        103 /* same as TK_SLASH, synopsis: r[P3]=r[P2]/r[P1] */
-#define OP_Remainder     104 /* same as TK_REM, synopsis: r[P3]=r[P2]%r[P1] */
-#define OP_Concat        105 /* same as TK_CONCAT, synopsis: r[P3]=r[P2]+r[P1] */
-#define OP_ReopenIdx     106 /* synopsis: root=P2 iDb=P3                   */
-#define OP_BitNot        107 /* same as TK_BITNOT, synopsis: r[P2]= ~r[P1] */
-#define OP_OpenRead      108 /* synopsis: root=P2 iDb=P3                   */
-#define OP_OpenWrite     109 /* synopsis: root=P2 iDb=P3                   */
-#define OP_String8       110 /* same as TK_STRING, synopsis: r[P2]='P4'    */
-#define OP_OpenDup       111
-#define OP_OpenAutoindex 112 /* synopsis: nColumn=P2                       */
-#define OP_OpenEphemeral 113 /* synopsis: nColumn=P2                       */
-#define OP_SorterOpen    114
-#define OP_SequenceTest  115 /* synopsis: if( cursor[P1].ctr++ ) pc = P2   */
-#define OP_OpenPseudo    116 /* synopsis: P3 columns in r[P2]              */
-#define OP_Close         117
-#define OP_ColumnsUsed   118
-#define OP_SeekHit       119 /* synopsis: seekHit=P2                       */
-#define OP_Sequence      120 /* synopsis: r[P2]=cursor[P1].ctr++           */
-#define OP_NewRowid      121 /* synopsis: r[P2]=rowid                      */
-#define OP_Insert        122 /* synopsis: intkey=r[P3] data=r[P2]          */
-#define OP_Delete        123
-#define OP_ResetCount    124
-#define OP_SorterCompare 125 /* synopsis: if key(P1)!=trim(r[P3],P4) goto P2 */
-#define OP_SorterData    126 /* synopsis: r[P2]=data                       */
-#define OP_RowData       127 /* synopsis: r[P2]=data                       */
-#define OP_Rowid         128 /* synopsis: r[P2]=rowid                      */
-#define OP_NullRow       129
-#define OP_SeekEnd       130
-#define OP_SorterInsert  131 /* synopsis: key=r[P2]                        */
-#define OP_IdxInsert     132 /* synopsis: key=r[P2]                        */
-#define OP_IdxDelete     133 /* synopsis: key=r[P2@P3]                     */
-#define OP_DeferredSeek  134 /* synopsis: Move P3 to P1.rowid if needed    */
-#define OP_IdxRowid      135 /* synopsis: r[P2]=rowid                      */
+#define OP_DecrJumpZero   59 /* jump, synopsis: if (--r[P1])==0 goto P2    */
+#define OP_IncrVacuum     60 /* jump                                       */
+#define OP_VNext          61 /* jump                                       */
+#define OP_Init           62 /* jump, synopsis: Start at P2                */
+#define OP_PureFunc       63 /* synopsis: r[P3]=func(r[P2@P5])             */
+#define OP_Function       64 /* synopsis: r[P3]=func(r[P2@P5])             */
+#define OP_Return         65
+#define OP_EndCoroutine   66
+#define OP_HaltIfNull     67 /* synopsis: if r[P3]=null halt               */
+#define OP_Halt           68
+#define OP_Integer        69 /* synopsis: r[P2]=P1                         */
+#define OP_Int64          70 /* synopsis: r[P2]=P4                         */
+#define OP_String         71 /* synopsis: r[P2]='P4' (len=P1)              */
+#define OP_Null           72 /* synopsis: r[P2..P3]=NULL                   */
+#define OP_SoftNull       73 /* synopsis: r[P1]=NULL                       */
+#define OP_Blob           74 /* synopsis: r[P2]=P4 (len=P1)                */
+#define OP_Variable       75 /* synopsis: r[P2]=parameter(P1,P4)           */
+#define OP_Move           76 /* synopsis: r[P2@P3]=r[P1@P3]                */
+#define OP_Copy           77 /* synopsis: r[P2@P3+1]=r[P1@P3+1]            */
+#define OP_SCopy          78 /* synopsis: r[P2]=r[P1]                      */
+#define OP_IntCopy        79 /* synopsis: r[P2]=r[P1]                      */
+#define OP_ResultRow      80 /* synopsis: output=r[P1@P2]                  */
+#define OP_CollSeq        81
+#define OP_AddImm         82 /* synopsis: r[P1]=r[P1]+P2                   */
+#define OP_RealAffinity   83
+#define OP_Cast           84 /* synopsis: affinity(r[P1])                  */
+#define OP_Permutation    85
+#define OP_Compare        86 /* synopsis: r[P1@P3] <-> r[P2@P3]            */
+#define OP_IsTrue         87 /* synopsis: r[P2] = coalesce(r[P1]==TRUE,P3) ^ P4 */
+#define OP_Offset         88 /* synopsis: r[P3] = sqlite_offset(P1)        */
+#define OP_Column         89 /* synopsis: r[P3]=PX                         */
+#define OP_Affinity       90 /* synopsis: affinity(r[P1@P2])               */
+#define OP_MakeRecord     91 /* synopsis: r[P3]=mkrec(r[P1@P2])            */
+#define OP_Count          92 /* synopsis: r[P2]=count()                    */
+#define OP_ReadCookie     93
+#define OP_SetCookie      94
+#define OP_ReopenIdx      95 /* synopsis: root=P2 iDb=P3                   */
+#define OP_OpenRead       96 /* synopsis: root=P2 iDb=P3                   */
+#define OP_OpenWrite      97 /* synopsis: root=P2 iDb=P3                   */
+#define OP_OpenDup        98
+#define OP_OpenAutoindex  99 /* synopsis: nColumn=P2                       */
+#define OP_OpenEphemeral 100 /* synopsis: nColumn=P2                       */
+#define OP_BitAnd        101 /* same as TK_BITAND, synopsis: r[P3]=r[P1]&r[P2] */
+#define OP_BitOr         102 /* same as TK_BITOR, synopsis: r[P3]=r[P1]|r[P2] */
+#define OP_ShiftLeft     103 /* same as TK_LSHIFT, synopsis: r[P3]=r[P2]<<r[P1] */
+#define OP_ShiftRight    104 /* same as TK_RSHIFT, synopsis: r[P3]=r[P2]>>r[P1] */
+#define OP_Add           105 /* same as TK_PLUS, synopsis: r[P3]=r[P1]+r[P2] */
+#define OP_Subtract      106 /* same as TK_MINUS, synopsis: r[P3]=r[P2]-r[P1] */
+#define OP_Multiply      107 /* same as TK_STAR, synopsis: r[P3]=r[P1]*r[P2] */
+#define OP_Divide        108 /* same as TK_SLASH, synopsis: r[P3]=r[P2]/r[P1] */
+#define OP_Remainder     109 /* same as TK_REM, synopsis: r[P3]=r[P2]%r[P1] */
+#define OP_Concat        110 /* same as TK_CONCAT, synopsis: r[P3]=r[P2]+r[P1] */
+#define OP_SorterOpen    111
+#define OP_BitNot        112 /* same as TK_BITNOT, synopsis: r[P2]= ~r[P1] */
+#define OP_SequenceTest  113 /* synopsis: if( cursor[P1].ctr++ ) pc = P2   */
+#define OP_OpenPseudo    114 /* synopsis: P3 columns in r[P2]              */
+#define OP_String8       115 /* same as TK_STRING, synopsis: r[P2]='P4'    */
+#define OP_Close         116
+#define OP_ColumnsUsed   117
+#define OP_SeekHit       118 /* synopsis: seekHit=P2                       */
+#define OP_Sequence      119 /* synopsis: r[P2]=cursor[P1].ctr++           */
+#define OP_NewRowid      120 /* synopsis: r[P2]=rowid                      */
+#define OP_Insert        121 /* synopsis: intkey=r[P3] data=r[P2]          */
+#define OP_Delete        122
+#define OP_ResetCount    123
+#define OP_SorterCompare 124 /* synopsis: if key(P1)!=trim(r[P3],P4) goto P2 */
+#define OP_SorterData    125 /* synopsis: r[P2]=data                       */
+#define OP_RowData       126 /* synopsis: r[P2]=data                       */
+#define OP_Rowid         127 /* synopsis: r[P2]=rowid                      */
+#define OP_NullRow       128
+#define OP_SeekEnd       129
+#define OP_SorterInsert  130 /* synopsis: key=r[P2]                        */
+#define OP_IdxInsert     131 /* synopsis: key=r[P2]                        */
+#define OP_IdxDelete     132 /* synopsis: key=r[P2@P3]                     */
+#define OP_DeferredSeek  133 /* synopsis: Move P3 to P1.rowid if needed    */
+#define OP_IdxRowid      134 /* synopsis: r[P2]=rowid                      */
+#define OP_FinishSeek    135
 #define OP_Destroy       136
 #define OP_Clear         137
 #define OP_ResetSorter   138
@@ -14988,12 +15337,12 @@ typedef struct VdbeOpList VdbeOpList;
 #define OP_LoadAnalysis  142
 #define OP_DropTable     143
 #define OP_DropIndex     144
-#define OP_Real          145 /* same as TK_FLOAT, synopsis: r[P2]=P4       */
-#define OP_DropTrigger   146
-#define OP_IntegrityCk   147
-#define OP_RowSetAdd     148 /* synopsis: rowset(P1)=r[P2]                 */
-#define OP_Param         149
-#define OP_FkCounter     150 /* synopsis: fkctr[P1]+=P2                    */
+#define OP_DropTrigger   145
+#define OP_IntegrityCk   146
+#define OP_RowSetAdd     147 /* synopsis: rowset(P1)=r[P2]                 */
+#define OP_Param         148
+#define OP_FkCounter     149 /* synopsis: fkctr[P1]+=P2                    */
+#define OP_Real          150 /* same as TK_FLOAT, synopsis: r[P2]=P4       */
 #define OP_MemMax        151 /* synopsis: r[P1]=max(r[P1],r[P2])           */
 #define OP_OffsetLimit   152 /* synopsis: if r[P1]>0 then r[P2]=r[P1]+max(0,r[P3]) else r[P2]=(-1) */
 #define OP_AggInverse    153 /* synopsis: accum=r[P3] inverse(r[P2@P5])    */
@@ -15002,20 +15351,23 @@ typedef struct VdbeOpList VdbeOpList;
 #define OP_AggValue      156 /* synopsis: r[P3]=value N=P2                 */
 #define OP_AggFinal      157 /* synopsis: accum=r[P1] N=P2                 */
 #define OP_Expire        158
-#define OP_TableLock     159 /* synopsis: iDb=P1 root=P2 write=P3          */
-#define OP_VBegin        160
-#define OP_VCreate       161
-#define OP_VDestroy      162
-#define OP_VOpen         163
-#define OP_VColumn       164 /* synopsis: r[P3]=vcolumn(P2)                */
-#define OP_VRename       165
-#define OP_Pagecount     166
-#define OP_MaxPgcnt      167
-#define OP_Trace         168
-#define OP_CursorHint    169
-#define OP_Noop          170
-#define OP_Explain       171
-#define OP_Abortable     172
+#define OP_CursorLock    159
+#define OP_CursorUnlock  160
+#define OP_TableLock     161 /* synopsis: iDb=P1 root=P2 write=P3          */
+#define OP_VBegin        162
+#define OP_VCreate       163
+#define OP_VDestroy      164
+#define OP_VOpen         165
+#define OP_VColumn       166 /* synopsis: r[P3]=vcolumn(P2)                */
+#define OP_VRename       167
+#define OP_Pagecount     168
+#define OP_MaxPgcnt      169
+#define OP_Trace         170
+#define OP_CursorHint    171
+#define OP_ReleaseReg    172 /* synopsis: release r[P1@P2] mask P3         */
+#define OP_Noop          173
+#define OP_Explain       174
+#define OP_Abortable     175
 
 /* Properties such as "out2" or "jump" that are specified in
 ** comments following the "case" for each opcode in the vdbe.c
@@ -15031,25 +15383,26 @@ typedef struct VdbeOpList VdbeOpList;
 /*   0 */ 0x00, 0x00, 0x00, 0x01, 0x01, 0x01, 0x00, 0x10,\
 /*   8 */ 0x00, 0x01, 0x00, 0x01, 0x01, 0x01, 0x03, 0x03,\
 /*  16 */ 0x01, 0x01, 0x03, 0x12, 0x03, 0x01, 0x09, 0x09,\
-/*  24 */ 0x09, 0x09, 0x09, 0x09, 0x09, 0x09, 0x09, 0x09,\
-/*  32 */ 0x01, 0x01, 0x01, 0x01, 0x01, 0x01, 0x01, 0x01,\
-/*  40 */ 0x01, 0x23, 0x0b, 0x26, 0x26, 0x01, 0x01, 0x03,\
+/*  24 */ 0x09, 0x09, 0x01, 0x09, 0x09, 0x09, 0x09, 0x09,\
+/*  32 */ 0x09, 0x01, 0x01, 0x01, 0x01, 0x01, 0x01, 0x01,\
+/*  40 */ 0x01, 0x01, 0x23, 0x26, 0x26, 0x0b, 0x01, 0x01,\
 /*  48 */ 0x03, 0x03, 0x03, 0x03, 0x0b, 0x0b, 0x0b, 0x0b,\
-/*  56 */ 0x0b, 0x0b, 0x01, 0x01, 0x01, 0x01, 0x00, 0x00,\
-/*  64 */ 0x00, 0x00, 0x02, 0x02, 0x08, 0x00, 0x10, 0x10,\
-/*  72 */ 0x10, 0x10, 0x00, 0x10, 0x10, 0x00, 0x00, 0x10,\
-/*  80 */ 0x10, 0x00, 0x00, 0x02, 0x02, 0x02, 0x00, 0x00,\
-/*  88 */ 0x12, 0x20, 0x00, 0x00, 0x00, 0x10, 0x10, 0x00,\
-/*  96 */ 0x26, 0x26, 0x26, 0x26, 0x26, 0x26, 0x26, 0x26,\
-/* 104 */ 0x26, 0x26, 0x00, 0x12, 0x00, 0x00, 0x10, 0x00,\
-/* 112 */ 0x00, 0x00, 0x00, 0x00, 0x00, 0x00, 0x00, 0x00,\
-/* 120 */ 0x10, 0x10, 0x00, 0x00, 0x00, 0x00, 0x00, 0x00,\
-/* 128 */ 0x10, 0x00, 0x00, 0x04, 0x04, 0x00, 0x00, 0x10,\
+/*  56 */ 0x0b, 0x0b, 0x01, 0x03, 0x01, 0x01, 0x01, 0x00,\
+/*  64 */ 0x00, 0x02, 0x02, 0x08, 0x00, 0x10, 0x10, 0x10,\
+/*  72 */ 0x10, 0x00, 0x10, 0x10, 0x00, 0x00, 0x10, 0x10,\
+/*  80 */ 0x00, 0x00, 0x02, 0x02, 0x02, 0x00, 0x00, 0x12,\
+/*  88 */ 0x20, 0x00, 0x00, 0x00, 0x10, 0x10, 0x00, 0x00,\
+/*  96 */ 0x00, 0x00, 0x00, 0x00, 0x00, 0x26, 0x26, 0x26,\
+/* 104 */ 0x26, 0x26, 0x26, 0x26, 0x26, 0x26, 0x26, 0x00,\
+/* 112 */ 0x12, 0x00, 0x00, 0x10, 0x00, 0x00, 0x00, 0x10,\
+/* 120 */ 0x10, 0x00, 0x00, 0x00, 0x00, 0x00, 0x00, 0x10,\
+/* 128 */ 0x00, 0x00, 0x04, 0x04, 0x00, 0x00, 0x10, 0x00,\
 /* 136 */ 0x10, 0x00, 0x00, 0x10, 0x00, 0x00, 0x00, 0x00,\
-/* 144 */ 0x00, 0x10, 0x00, 0x00, 0x06, 0x10, 0x00, 0x04,\
+/* 144 */ 0x00, 0x00, 0x00, 0x06, 0x10, 0x00, 0x10, 0x04,\
 /* 152 */ 0x1a, 0x00, 0x00, 0x00, 0x00, 0x00, 0x00, 0x00,\
-/* 160 */ 0x00, 0x00, 0x00, 0x00, 0x00, 0x00, 0x10, 0x10,\
-/* 168 */ 0x00, 0x00, 0x00, 0x00, 0x00,}
+/* 160 */ 0x00, 0x00, 0x00, 0x00, 0x00, 0x00, 0x00, 0x00,\
+/* 168 */ 0x10, 0x10, 0x00, 0x00, 0x00, 0x00, 0x00, 0x00,\
+}
 
 /* The sqlite3P2Values() routine is able to run faster if it knows
 ** the value of the largest JUMP opcode.  The smaller the maximum
@@ -15057,7 +15410,7 @@ typedef struct VdbeOpList VdbeOpList;
 ** generated this include file strives to group all JUMP opcodes
 ** together near the beginning of the list.
 */
-#define SQLITE_MX_JUMP_OPCODE  61  /* Maximum JUMP opcode */
+#define SQLITE_MX_JUMP_OPCODE  62  /* Maximum JUMP opcode */
 
 /************** End of opcodes.h *********************************************/
 /************** Continuing where we left off in vdbe.h ***********************/
@@ -15073,6 +15426,7 @@ typedef struct VdbeOpList VdbeOpList;
 ** for a description of what each of these routines does.
 */
 SQLITE_PRIVATE Vdbe *sqlite3VdbeCreate(Parse*);
+SQLITE_PRIVATE Parse *sqlite3VdbeParser(Vdbe*);
 SQLITE_PRIVATE int sqlite3VdbeAddOp0(Vdbe*,int);
 SQLITE_PRIVATE int sqlite3VdbeAddOp1(Vdbe*,int,int);
 SQLITE_PRIVATE int sqlite3VdbeAddOp2(Vdbe*,int,int,int);
@@ -15083,6 +15437,7 @@ SQLITE_PRIVATE int sqlite3VdbeAddOp3(Vdbe*,int,int,int,int);
 SQLITE_PRIVATE int sqlite3VdbeAddOp4(Vdbe*,int,int,int,int,const char *zP4,int);
 SQLITE_PRIVATE int sqlite3VdbeAddOp4Dup8(Vdbe*,int,int,int,int,const u8*,int);
 SQLITE_PRIVATE int sqlite3VdbeAddOp4Int(Vdbe*,int,int,int,int,int);
+SQLITE_PRIVATE int sqlite3VdbeAddFunctionCall(Parse*,int,int,int,int,const FuncDef*,int);
 SQLITE_PRIVATE void sqlite3VdbeEndCoroutine(Vdbe*,int);
 #if defined(SQLITE_DEBUG) && !defined(SQLITE_TEST_REALLOC_STRESS)
 SQLITE_PRIVATE   void sqlite3VdbeVerifyNoMallocRequired(Vdbe *p, int N);
@@ -15116,14 +15471,19 @@ SQLITE_PRIVATE   void sqlite3ExplainBreakpoint(const char*,const char*);
 # define sqlite3ExplainBreakpoint(A,B) /*no-op*/
 #endif
 SQLITE_PRIVATE void sqlite3VdbeAddParseSchemaOp(Vdbe*,int,char*);
-SQLITE_PRIVATE void sqlite3VdbeChangeOpcode(Vdbe*, u32 addr, u8);
-SQLITE_PRIVATE void sqlite3VdbeChangeP1(Vdbe*, u32 addr, int P1);
-SQLITE_PRIVATE void sqlite3VdbeChangeP2(Vdbe*, u32 addr, int P2);
-SQLITE_PRIVATE void sqlite3VdbeChangeP3(Vdbe*, u32 addr, int P3);
+SQLITE_PRIVATE void sqlite3VdbeChangeOpcode(Vdbe*, int addr, u8);
+SQLITE_PRIVATE void sqlite3VdbeChangeP1(Vdbe*, int addr, int P1);
+SQLITE_PRIVATE void sqlite3VdbeChangeP2(Vdbe*, int addr, int P2);
+SQLITE_PRIVATE void sqlite3VdbeChangeP3(Vdbe*, int addr, int P3);
 SQLITE_PRIVATE void sqlite3VdbeChangeP5(Vdbe*, u16 P5);
 SQLITE_PRIVATE void sqlite3VdbeJumpHere(Vdbe*, int addr);
 SQLITE_PRIVATE int sqlite3VdbeChangeToNoop(Vdbe*, int addr);
 SQLITE_PRIVATE int sqlite3VdbeDeletePriorOpcode(Vdbe*, u8 op);
+#ifdef SQLITE_DEBUG
+SQLITE_PRIVATE   void sqlite3VdbeReleaseRegisters(Parse*,int addr, int n, u32 mask, int);
+#else
+# define sqlite3VdbeReleaseRegisters(P,A,N,M,F)
+#endif
 SQLITE_PRIVATE void sqlite3VdbeChangeP4(Vdbe*, int addr, const char *zP4, int N);
 SQLITE_PRIVATE void sqlite3VdbeAppendP4(Vdbe*, void *pP4, int p4type);
 SQLITE_PRIVATE void sqlite3VdbeSetP4KeyInfo(Parse*, Index*);
@@ -15172,9 +15532,8 @@ SQLITE_PRIVATE UnpackedRecord *sqlite3VdbeAllocUnpackedRecord(KeyInfo*);
 typedef int (*RecordCompare)(int,const void*,UnpackedRecord*);
 SQLITE_PRIVATE RecordCompare sqlite3VdbeFindCompare(UnpackedRecord*);
 
-#ifndef SQLITE_OMIT_TRIGGER
 SQLITE_PRIVATE void sqlite3VdbeLinkSubProgram(Vdbe *, SubProgram *);
-#endif
+SQLITE_PRIVATE int sqlite3VdbeHasSubProgram(Vdbe*);
 
 SQLITE_PRIVATE int sqlite3NotPureFunc(sqlite3_context*);
 
@@ -15488,7 +15847,7 @@ SQLITE_PRIVATE u32 sqlite3PagerDataVersion(Pager*);
 SQLITE_PRIVATE   int sqlite3PagerRefcount(Pager*);
 #endif
 SQLITE_PRIVATE int sqlite3PagerMemUsed(Pager*);
-SQLITE_PRIVATE const char *sqlite3PagerFilename(Pager*, int);
+SQLITE_PRIVATE const char *sqlite3PagerFilename(const Pager*, int);
 SQLITE_PRIVATE sqlite3_vfs *sqlite3PagerVfs(Pager*);
 SQLITE_PRIVATE sqlite3_file *sqlite3PagerFile(Pager*);
 SQLITE_PRIVATE sqlite3_file *sqlite3PagerJrnlFile(Pager*);
@@ -16076,6 +16435,7 @@ SQLITE_PRIVATE void sqlite3OsCloseFree(sqlite3_file *);
 #define MUTEX_LOGIC(X)
 #else
 #define MUTEX_LOGIC(X)            X
+SQLITE_API int sqlite3_mutex_held(sqlite3_mutex*);
 #endif /* defined(SQLITE_MUTEX_OMIT) */
 
 /************** End of mutex.h ***********************************************/
@@ -16207,15 +16567,47 @@ struct Schema {
 ** is shared by multiple database connections.  Therefore, while parsing
 ** schema information, the Lookaside.bEnabled flag is cleared so that
 ** lookaside allocations are not used to construct the schema objects.
+**
+** New lookaside allocations are only allowed if bDisable==0.  When
+** bDisable is greater than zero, sz is set to zero which effectively
+** disables lookaside without adding a new test for the bDisable flag
+** in a performance-critical path.  sz should be set by to szTrue whenever
+** bDisable changes back to zero.
+**
+** Lookaside buffers are initially held on the pInit list.  As they are
+** used and freed, they are added back to the pFree list.  New allocations
+** come off of pFree first, then pInit as a fallback.  This dual-list
+** allows use to compute a high-water mark - the maximum number of allocations
+** outstanding at any point in the past - by subtracting the number of
+** allocations on the pInit list from the total number of allocations.
+**
+** Enhancement on 2019-12-12:  Two-size-lookaside
+** The default lookaside configuration is 100 slots of 1200 bytes each.
+** The larger slot sizes are important for performance, but they waste
+** a lot of space, as most lookaside allocations are less than 128 bytes.
+** The two-size-lookaside enhancement breaks up the lookaside allocation
+** into two pools:  One of 128-byte slots and the other of the default size
+** (1200-byte) slots.   Allocations are filled from the small-pool first,
+** failing over to the full-size pool if that does not work.  Thus more
+** lookaside slots are available while also using less memory.
+** This enhancement can be omitted by compiling with
+** SQLITE_OMIT_TWOSIZE_LOOKASIDE.
 */
 struct Lookaside {
   u32 bDisable;           /* Only operate the lookaside when zero */
   u16 sz;                 /* Size of each buffer in bytes */
+  u16 szTrue;             /* True value of sz, even if disabled */
   u8 bMalloced;           /* True if pStart obtained from sqlite3_malloc() */
   u32 nSlot;              /* Number of lookaside slots allocated */
   u32 anStat[3];          /* 0: hits.  1: size misses.  2: full misses */
   LookasideSlot *pInit;   /* List of buffers not previously used */
   LookasideSlot *pFree;   /* List of available buffers */
+#ifndef SQLITE_OMIT_TWOSIZE_LOOKASIDE
+  LookasideSlot *pSmallInit; /* List of small buffers not prediously used */
+  LookasideSlot *pSmallFree; /* List of available small buffers */
+  void *pMiddle;          /* First byte past end of full-size buffers and
+                          ** the first byte of LOOKASIDE_SMALL buffers */
+#endif /* SQLITE_OMIT_TWOSIZE_LOOKASIDE */
   void *pStart;           /* First byte of available memory space */
   void *pEnd;             /* First byte past end of available space */
 };
@@ -16223,6 +16615,17 @@ struct LookasideSlot {
   LookasideSlot *pNext;    /* Next buffer in the list of free buffers */
 };
 
+#define DisableLookaside  db->lookaside.bDisable++;db->lookaside.sz=0
+#define EnableLookaside   db->lookaside.bDisable--;\
+   db->lookaside.sz=db->lookaside.bDisable?0:db->lookaside.szTrue
+
+/* Size of the smaller allocations in two-size lookside */
+#ifdef SQLITE_OMIT_TWOSIZE_LOOKASIDE
+#  define LOOKASIDE_SMALL           0
+#else
+#  define LOOKASIDE_SMALL         128
+#endif
+
 /*
 ** A hash table for built-in function definitions.  (Application-defined
 ** functions use a regular table table from hash.h.)
@@ -16334,6 +16737,7 @@ struct sqlite3 {
     unsigned orphanTrigger : 1; /* Last statement is orphaned TEMP trigger */
     unsigned imposterTable : 1; /* Building an imposter table */
     unsigned reopenMemdb : 1;   /* ATTACH is really a reopen using MemDB */
+    char **azInit;              /* "type", "name", and "tbl_name" columns */
   } init;
   int nVdbeActive;              /* Number of VDBEs currently running */
   int nVdbeRead;                /* Number of active VDBEs that read or write */
@@ -16430,6 +16834,13 @@ struct sqlite3 {
 #define SCHEMA_ENC(db) ((db)->aDb[0].pSchema->enc)
 #define ENC(db)        ((db)->enc)
 
+/*
+** A u64 constant where the lower 32 bits are all zeros.  Only the
+** upper 32 bits are included in the argument.  Necessary because some
+** C-compilers still do not accept LL integer literals.
+*/
+#define HI(X)  ((u64)(X)<<32)
+
 /*
 ** Possible values for the sqlite3.flags.
 **
@@ -16445,9 +16856,8 @@ struct sqlite3 {
 #define SQLITE_CkptFullFSync  0x00000010  /* Use full fsync for checkpoint */
 #define SQLITE_CacheSpill     0x00000020  /* OK to spill pager cache */
 #define SQLITE_ShortColNames  0x00000040  /* Show short columns names */
-#define SQLITE_CountRows      0x00000080  /* Count rows changed by INSERT, */
-                                          /*   DELETE, or UPDATE and return */
-                                          /*   the count using a callback. */
+#define SQLITE_TrustedSchema  0x00000080  /* Allow unsafe functions and
+                                          ** vtabs in the schema definition */
 #define SQLITE_NullCallback   0x00000100  /* Invoke the callback once if the */
                                           /*   result set is empty */
 #define SQLITE_IgnoreChecks   0x00000200  /* Do not enforce check constraints */
@@ -16470,16 +16880,21 @@ struct sqlite3 {
 #define SQLITE_LegacyAlter    0x04000000  /* Legacy ALTER TABLE behaviour */
 #define SQLITE_NoSchemaError  0x08000000  /* Do not report schema parse errors*/
 #define SQLITE_Defensive      0x10000000  /* Input SQL is likely hostile */
+#define SQLITE_DqsDDL         0x20000000  /* dbl-quoted strings allowed in DDL*/
+#define SQLITE_DqsDML         0x40000000  /* dbl-quoted strings allowed in DML*/
+#define SQLITE_EnableView     0x80000000  /* Enable the use of views */
+#define SQLITE_CountRows      HI(0x00001) /* Count rows changed by INSERT, */
+                                          /*   DELETE, or UPDATE and return */
+                                          /*   the count using a callback. */
 
 /* Flags used only if debugging */
-#define HI(X)  ((u64)(X)<<32)
 #ifdef SQLITE_DEBUG
-#define SQLITE_SqlTrace       HI(0x0001)  /* Debug print SQL as it executes */
-#define SQLITE_VdbeListing    HI(0x0002)  /* Debug listings of VDBE progs */
-#define SQLITE_VdbeTrace      HI(0x0004)  /* True to trace VDBE execution */
-#define SQLITE_VdbeAddopTrace HI(0x0008)  /* Trace sqlite3VdbeAddOp() calls */
-#define SQLITE_VdbeEQP        HI(0x0010)  /* Debug EXPLAIN QUERY PLAN */
-#define SQLITE_ParserTrace    HI(0x0020)  /* PRAGMA parser_trace=ON */
+#define SQLITE_SqlTrace       HI(0x0100000) /* Debug print SQL as it executes */
+#define SQLITE_VdbeListing    HI(0x0200000) /* Debug listings of VDBE progs */
+#define SQLITE_VdbeTrace      HI(0x0400000) /* True to trace VDBE execution */
+#define SQLITE_VdbeAddopTrace HI(0x0800000) /* Trace sqlite3VdbeAddOp() calls */
+#define SQLITE_VdbeEQP        HI(0x1000000) /* Debug EXPLAIN QUERY PLAN */
+#define SQLITE_ParserTrace    HI(0x2000000) /* PRAGMA parser_trace=ON */
 #endif
 
 /*
@@ -16490,6 +16905,7 @@ struct sqlite3 {
 #define DBFLAG_Vacuum         0x0004  /* Currently in a VACUUM */
 #define DBFLAG_VacuumInto     0x0008  /* Currently running VACUUM INTO */
 #define DBFLAG_SchemaKnownOk  0x0010  /* Schema is known to be valid */
+#define DBFLAG_InternalFunc   0x0020  /* Allow use of internal functions */
 
 /*
 ** Bits of the sqlite3.dbOptFlags field that are used by the
@@ -16507,8 +16923,8 @@ struct sqlite3 {
 #define SQLITE_OmitNoopJoin   0x0100   /* Omit unused tables in joins */
 #define SQLITE_CountOfView    0x0200   /* The count-of-view optimization */
 #define SQLITE_CursorHints    0x0400   /* Add OP_CursorHint opcodes */
-#define SQLITE_Stat34         0x0800   /* Use STAT3 or STAT4 data */
-   /* TH3 expects the Stat34  ^^^^^^ value to be 0x0800.  Don't change it */
+#define SQLITE_Stat4          0x0800   /* Use STAT4 data */
+   /* TH3 expects the Stat4   ^^^^^^ value to be 0x0800.  Don't change it */
 #define SQLITE_PushDown       0x1000   /* The push-down optimization */
 #define SQLITE_SimplifyJoin   0x2000   /* Convert LEFT JOIN to JOIN */
 #define SQLITE_SkipScan       0x4000   /* Skip-scans */
@@ -16596,6 +17012,8 @@ struct FuncDestructor {
 **     SQLITE_FUNC_LENGTH    ==  OPFLAG_LENGTHARG
 **     SQLITE_FUNC_TYPEOF    ==  OPFLAG_TYPEOFARG
 **     SQLITE_FUNC_CONSTANT  ==  SQLITE_DETERMINISTIC from the API
+**     SQLITE_FUNC_DIRECT    ==  SQLITE_DIRECTONLY from the API
+**     SQLITE_FUNC_UNSAFE    ==  SQLITE_INNOCUOUS
 **     SQLITE_FUNC_ENCMASK   depends on SQLITE_UTF* macros in the API
 */
 #define SQLITE_FUNC_ENCMASK  0x0003 /* SQLITE_UTF8, SQLITE_UTF16BE or UTF16LE */
@@ -16612,10 +17030,22 @@ struct FuncDestructor {
 #define SQLITE_FUNC_MINMAX   0x1000 /* True for min() and max() aggregates */
 #define SQLITE_FUNC_SLOCHNG  0x2000 /* "Slow Change". Value constant during a
                                     ** single query - might change over time */
-#define SQLITE_FUNC_AFFINITY 0x4000 /* Built-in affinity() function */
+#define SQLITE_FUNC_TEST     0x4000 /* Built-in testing functions */
 #define SQLITE_FUNC_OFFSET   0x8000 /* Built-in sqlite_offset() function */
 #define SQLITE_FUNC_WINDOW   0x00010000 /* Built-in window-only function */
 #define SQLITE_FUNC_INTERNAL 0x00040000 /* For use by NestedParse() only */
+#define SQLITE_FUNC_DIRECT   0x00080000 /* Not for use in TRIGGERs or VIEWs */
+#define SQLITE_FUNC_SUBTYPE  0x00100000 /* Result likely to have sub-type */
+#define SQLITE_FUNC_UNSAFE   0x00200000 /* Function has side effects */
+#define SQLITE_FUNC_INLINE   0x00400000 /* Functions implemented in-line */
+
+/* Identifier numbers for each in-line function */
+#define INLINEFUNC_coalesce             0
+#define INLINEFUNC_implies_nonnull_row  1
+#define INLINEFUNC_expr_implies_expr    2
+#define INLINEFUNC_expr_compare         3      
+#define INLINEFUNC_affinity             4
+#define INLINEFUNC_unlikely            99  /* Default case */
 
 /*
 ** The following three macros, FUNCTION(), LIKEFUNC() and AGGREGATE() are
@@ -16631,6 +17061,22 @@ struct FuncDestructor {
 **   VFUNCTION(zName, nArg, iArg, bNC, xFunc)
 **     Like FUNCTION except it omits the SQLITE_FUNC_CONSTANT flag.
 **
+**   SFUNCTION(zName, nArg, iArg, bNC, xFunc)
+**     Like FUNCTION except it omits the SQLITE_FUNC_CONSTANT flag and
+**     adds the SQLITE_DIRECTONLY flag.
+**
+**   INLINE_FUNC(zName, nArg, iFuncId, mFlags)
+**     zName is the name of a function that is implemented by in-line
+**     byte code rather than by the usual callbacks. The iFuncId
+**     parameter determines the function id.  The mFlags parameter is
+**     optional SQLITE_FUNC_ flags for this function.
+**
+**   TEST_FUNC(zName, nArg, iFuncId, mFlags)
+**     zName is the name of a test-only function implemented by in-line
+**     byte code rather than by the usual callbacks. The iFuncId
+**     parameter determines the function id.  The mFlags parameter is
+**     optional SQLITE_FUNC_ flags for this function.
+**
 **   DFUNCTION(zName, nArg, iArg, bNC, xFunc)
 **     Like FUNCTION except it omits the SQLITE_FUNC_CONSTANT flag and
 **     adds the SQLITE_FUNC_SLOCHNG flag.  Used for date & time functions
@@ -16670,6 +17116,16 @@ struct FuncDestructor {
 #define VFUNCTION(zName, nArg, iArg, bNC, xFunc) \
   {nArg, SQLITE_UTF8|(bNC*SQLITE_FUNC_NEEDCOLL), \
    SQLITE_INT_TO_PTR(iArg), 0, xFunc, 0, 0, 0, #zName, {0} }
+#define SFUNCTION(zName, nArg, iArg, bNC, xFunc) \
+  {nArg, SQLITE_UTF8|SQLITE_DIRECTONLY|SQLITE_FUNC_UNSAFE, \
+   SQLITE_INT_TO_PTR(iArg), 0, xFunc, 0, 0, 0, #zName, {0} }
+#define INLINE_FUNC(zName, nArg, iArg, mFlags) \
+  {nArg, SQLITE_UTF8|SQLITE_FUNC_INLINE|SQLITE_FUNC_CONSTANT|(mFlags), \
+   SQLITE_INT_TO_PTR(iArg), 0, noopFunc, 0, 0, 0, #zName, {0} }
+#define TEST_FUNC(zName, nArg, iArg, mFlags) \
+  {nArg, SQLITE_UTF8|SQLITE_FUNC_INTERNAL|SQLITE_FUNC_TEST| \
+         SQLITE_FUNC_INLINE|SQLITE_FUNC_CONSTANT|(mFlags), \
+   SQLITE_INT_TO_PTR(iArg), 0, noopFunc, 0, 0, 0, #zName, {0} }
 #define DFUNCTION(zName, nArg, iArg, bNC, xFunc) \
   {nArg, SQLITE_FUNC_SLOCHNG|SQLITE_UTF8, \
    0, 0, xFunc, 0, 0, 0, #zName, {0} }
@@ -16685,12 +17141,6 @@ struct FuncDestructor {
 #define LIKEFUNC(zName, nArg, arg, flags) \
   {nArg, SQLITE_FUNC_CONSTANT|SQLITE_UTF8|flags, \
    (void *)arg, 0, likeFunc, 0, 0, 0, #zName, {0} }
-#define AGGREGATE(zName, nArg, arg, nc, xStep, xFinal, xValue) \
-  {nArg, SQLITE_UTF8|(nc*SQLITE_FUNC_NEEDCOLL), \
-   SQLITE_INT_TO_PTR(arg), 0, xStep,xFinal,xValue,0,#zName, {0}}
-#define AGGREGATE2(zName, nArg, arg, nc, xStep, xFinal, extraFlags) \
-  {nArg, SQLITE_UTF8|(nc*SQLITE_FUNC_NEEDCOLL)|extraFlags, \
-   SQLITE_INT_TO_PTR(arg), 0, xStep,xFinal,xFinal,0,#zName, {0}}
 #define WAGGREGATE(zName, nArg, arg, nc, xStep, xFinal, xValue, xInverse, f) \
   {nArg, SQLITE_UTF8|(nc*SQLITE_FUNC_NEEDCOLL)|f, \
    SQLITE_INT_TO_PTR(arg), 0, xStep,xFinal,xValue,xInverse,#zName, {0}}
@@ -16729,32 +17179,52 @@ struct Savepoint {
 struct Module {
   const sqlite3_module *pModule;       /* Callback pointers */
   const char *zName;                   /* Name passed to create_module() */
+  int nRefModule;                      /* Number of pointers to this object */
   void *pAux;                          /* pAux passed to create_module() */
   void (*xDestroy)(void *);            /* Module destructor function */
   Table *pEpoTab;                      /* Eponymous table for this module */
 };
 
 /*
-** information about each column of an SQL table is held in an instance
-** of this structure.
+** Information about each column of an SQL table is held in an instance
+** of the Column structure, in the Table.aCol[] array.
+**
+** Definitions:
+**
+**   "table column index"     This is the index of the column in the
+**                            Table.aCol[] array, and also the index of
+**                            the column in the original CREATE TABLE stmt.
+**
+**   "storage column index"   This is the index of the column in the
+**                            record BLOB generated by the OP_MakeRecord
+**                            opcode.  The storage column index is less than
+**                            or equal to the table column index.  It is
+**                            equal if and only if there are no VIRTUAL
+**                            columns to the left.
 */
 struct Column {
   char *zName;     /* Name of this column, \000, then the type */
-  Expr *pDflt;     /* Default value of this column */
+  Expr *pDflt;     /* Default value or GENERATED ALWAYS AS value */
   char *zColl;     /* Collating sequence.  If NULL, use the default */
   u8 notNull;      /* An OE_ code for handling a NOT NULL constraint */
   char affinity;   /* One of the SQLITE_AFF_... values */
   u8 szEst;        /* Estimated size of value in this column. sizeof(INT)==1 */
-  u8 colFlags;     /* Boolean properties.  See COLFLAG_ defines below */
+  u16 colFlags;    /* Boolean properties.  See COLFLAG_ defines below */
 };
 
 /* Allowed values for Column.colFlags:
 */
-#define COLFLAG_PRIMKEY  0x0001    /* Column is part of the primary key */
-#define COLFLAG_HIDDEN   0x0002    /* A hidden column in a virtual table */
-#define COLFLAG_HASTYPE  0x0004    /* Type name follows column name */
-#define COLFLAG_UNIQUE   0x0008    /* Column def contains "UNIQUE" or "PK" */
+#define COLFLAG_PRIMKEY   0x0001   /* Column is part of the primary key */
+#define COLFLAG_HIDDEN    0x0002   /* A hidden column in a virtual table */
+#define COLFLAG_HASTYPE   0x0004   /* Type name follows column name */
+#define COLFLAG_UNIQUE    0x0008   /* Column def contains "UNIQUE" or "PK" */
 #define COLFLAG_SORTERREF 0x0010   /* Use sorter-refs with this column */
+#define COLFLAG_VIRTUAL   0x0020   /* GENERATED ALWAYS AS ... VIRTUAL */
+#define COLFLAG_STORED    0x0040   /* GENERATED ALWAYS AS ... STORED */
+#define COLFLAG_NOTAVAIL  0x0080   /* STORED column not yet calculated */
+#define COLFLAG_BUSY      0x0100   /* Blocks recursion on GENERATED columns */
+#define COLFLAG_GENERATED 0x0060   /* Combo: _STORED, _VIRTUAL */
+#define COLFLAG_NOINSERT  0x0062   /* Combo: _HIDDEN, _STORED, _VIRTUAL */
 
 /*
 ** A "Collating Sequence" is defined by an instance of the following
@@ -16794,11 +17264,12 @@ struct CollSeq {
 ** Note also that the numeric types are grouped together so that testing
 ** for a numeric type is a single comparison.  And the BLOB type is first.
 */
-#define SQLITE_AFF_BLOB     'A'
-#define SQLITE_AFF_TEXT     'B'
-#define SQLITE_AFF_NUMERIC  'C'
-#define SQLITE_AFF_INTEGER  'D'
-#define SQLITE_AFF_REAL     'E'
+#define SQLITE_AFF_NONE     0x40  /* '@' */
+#define SQLITE_AFF_BLOB     0x41  /* 'A' */
+#define SQLITE_AFF_TEXT     0x42  /* 'B' */
+#define SQLITE_AFF_NUMERIC  0x43  /* 'C' */
+#define SQLITE_AFF_INTEGER  0x44  /* 'D' */
+#define SQLITE_AFF_REAL     0x45  /* 'E' */
 
 #define sqlite3IsNumericAffinity(X)  ((X)>=SQLITE_AFF_NUMERIC)
 
@@ -16871,10 +17342,17 @@ struct VTable {
   sqlite3_vtab *pVtab;      /* Pointer to vtab instance */
   int nRef;                 /* Number of pointers to this structure */
   u8 bConstraint;           /* True if constraints are supported */
+  u8 eVtabRisk;             /* Riskiness of allowing hacker access */
   int iSavepoint;           /* Depth of the SAVEPOINT stack */
   VTable *pNext;            /* Next in linked list (see above) */
 };
 
+/* Allowed values for VTable.eVtabRisk
+*/
+#define SQLITE_VTABRISK_Low          0
+#define SQLITE_VTABRISK_Normal       1
+#define SQLITE_VTABRISK_High         2
+
 /*
 ** The schema for each SQL table and view is represented in memory
 ** by an instance of the following structure.
@@ -16893,6 +17371,7 @@ struct Table {
   u32 tabFlags;        /* Mask of TF_* values */
   i16 iPKey;           /* If not negative, use aCol[iPKey] as the rowid */
   i16 nCol;            /* Number of columns in this table */
+  i16 nNVCol;          /* Number of columns that are not VIRTUAL */
   LogEst nRowLogEst;   /* Estimated rows in table - from sqlite_stat1 table */
   LogEst szTabRow;     /* Estimated size of each table row in bytes */
 #ifdef SQLITE_ENABLE_COSTMULT
@@ -16919,20 +17398,28 @@ struct Table {
 ** followed by non-hidden columns.  Example:  "CREATE VIRTUAL TABLE x USING
 ** vtab1(a HIDDEN, b);".  Since "b" is a non-hidden column but "a" is hidden,
 ** the TF_OOOHidden attribute would apply in this case.  Such tables require
-** special handling during INSERT processing.
+** special handling during INSERT processing. The "OOO" means "Out Of Order".
+**
+** Constraints:
+**
+**         TF_HasVirtual == COLFLAG_Virtual
+**         TF_HasStored  == COLFLAG_Stored
 */
 #define TF_Readonly        0x0001    /* Read-only system table */
 #define TF_Ephemeral       0x0002    /* An ephemeral table */
 #define TF_HasPrimaryKey   0x0004    /* Table has a primary key */
 #define TF_Autoincrement   0x0008    /* Integer primary key is autoincrement */
 #define TF_HasStat1        0x0010    /* nRowLogEst set from sqlite_stat1 */
-#define TF_WithoutRowid    0x0020    /* No rowid.  PRIMARY KEY is the key */
-#define TF_NoVisibleRowid  0x0040    /* No user-visible "rowid" column */
-#define TF_OOOHidden       0x0080    /* Out-of-Order hidden columns */
+#define TF_HasVirtual      0x0020    /* Has one or more VIRTUAL columns */
+#define TF_HasStored       0x0040    /* Has one or more STORED columns */
+#define TF_HasGenerated    0x0060    /* Combo: HasVirtual + HasStored */
+#define TF_WithoutRowid    0x0080    /* No rowid.  PRIMARY KEY is the key */
 #define TF_StatsUsed       0x0100    /* Query planner decisions affected by
                                      ** Index.aiRowLogEst[] values */
-#define TF_HasNotNull      0x0200    /* Contains NOT NULL constraints */
-#define TF_Shadow          0x0400    /* True for a shadow table */
+#define TF_NoVisibleRowid  0x0200    /* No user-visible "rowid" column */
+#define TF_OOOHidden       0x0400    /* Out-of-Order hidden columns */
+#define TF_HasNotNull      0x0800    /* Contains NOT NULL constraints */
+#define TF_Shadow          0x1000    /* True for a shadow table */
 
 /*
 ** Test to see whether or not a table is a virtual table.  This is
@@ -17066,10 +17553,16 @@ struct KeyInfo {
   u16 nKeyField;      /* Number of key columns in the index */
   u16 nAllField;      /* Total columns, including key plus others */
   sqlite3 *db;        /* The database connection */
-  u8 *aSortOrder;     /* Sort order for each column. */
+  u8 *aSortFlags;     /* Sort order for each column. */
   CollSeq *aColl[1];  /* Collating sequence for each term of the key */
 };
 
+/*
+** Allowed bit values for entries in the KeyInfo.aSortFlags[] array.
+*/
+#define KEYINFO_ORDER_DESC    0x01    /* DESC sort order */
+#define KEYINFO_ORDER_BIGNULL 0x02    /* NULL is larger than any other value */
+
 /*
 ** This object holds a record which has been parsed out into individual
 ** fields, for the purposes of doing a comparison.
@@ -17176,7 +17669,9 @@ struct Index {
   unsigned noSkipScan:1;   /* Do not try to use skip-scan if true */
   unsigned hasStat1:1;     /* aiRowLogEst values come from sqlite_stat1 */
   unsigned bNoQuery:1;     /* Do not use this index to optimize queries */
-#ifdef SQLITE_ENABLE_STAT3_OR_STAT4
+  unsigned bAscKeyBug:1;   /* True if the bba7b69f9849b5bf bug applies */
+  unsigned bHasVCol:1;     /* Index references one or more VIRTUAL columns */
+#ifdef SQLITE_ENABLE_STAT4
   int nSample;             /* Number of elements in aSample[] */
   int nSampleCol;          /* Size of IndexSample.anEq[] and so on */
   tRowcnt *aAvgEq;         /* Average nEq values for keys not in aSample */
@@ -17208,7 +17703,7 @@ struct Index {
 #define XN_EXPR      (-2)     /* Indexed column is an expression */
 
 /*
-** Each sample stored in the sqlite_stat3 table is represented in memory
+** Each sample stored in the sqlite_stat4 table is represented in memory
 ** using a structure of this type.  See documentation at the top of the
 ** analyze.c source file for additional information.
 */
@@ -17366,7 +17861,11 @@ typedef int ynVar;
 */
 struct Expr {
   u8 op;                 /* Operation performed by this node */
-  char affinity;         /* The affinity of the column or 0 if not a column */
+  char affExpr;          /* affinity, or RAISE type */
+  u8 op2;                /* TK_REGISTER/TK_TRUTH: original value of Expr.op
+                         ** TK_COLUMN: the value of p5 for OP_Column
+                         ** TK_AGG_FUNCTION: nesting depth
+                         ** TK_FUNCTION: NC_SelfRef flag if needs OP_PureFunc */
   u32 flags;             /* Various flags.  EP_* See below */
   union {
     char *zToken;          /* Token value. Zero terminated and dequoted */
@@ -17397,20 +17896,19 @@ struct Expr {
                          ** TK_REGISTER: register number
                          ** TK_TRIGGER: 1 -> new, 0 -> old
                          ** EP_Unlikely:  134217728 times likelihood
+                         ** TK_IN: ephemerial table holding RHS
+                         ** TK_SELECT_COLUMN: Number of columns on the LHS
                          ** TK_SELECT: 1st register of result vector */
   ynVar iColumn;         /* TK_COLUMN: column index.  -1 for rowid.
                          ** TK_VARIABLE: variable number (always >= 1).
                          ** TK_SELECT_COLUMN: column of the result vector */
   i16 iAgg;              /* Which entry in pAggInfo->aCol[] or ->aFunc[] */
   i16 iRightJoinTable;   /* If EP_FromJoin, the right table of the join */
-  u8 op2;                /* TK_REGISTER: original value of Expr.op
-                         ** TK_COLUMN: the value of p5 for OP_Column
-                         ** TK_AGG_FUNCTION: nesting depth */
   AggInfo *pAggInfo;     /* Used by TK_AGG_COLUMN and TK_AGG_FUNCTION */
   union {
     Table *pTab;           /* TK_COLUMN: Table containing column. Can be NULL
                            ** for a column of an index on an expression */
-    Window *pWin;          /* TK_FUNCTION: Window definition for the func */
+    Window *pWin;          /* EP_WinFunc: Window/Filter defn for a function */
     struct {               /* TK_IN, TK_SELECT, and TK_EXISTS */
       int iAddr;             /* Subroutine entry address */
       int regReturn;         /* Register used to hold return address */
@@ -17425,34 +17923,37 @@ struct Expr {
 **          EP_Agg == NC_HasAgg == SF_HasAgg
 **          EP_Win == NC_HasWin
 */
-#define EP_FromJoin  0x000001 /* Originates in ON/USING clause of outer join */
-#define EP_Distinct  0x000002 /* Aggregate function with DISTINCT keyword */
-#define EP_HasFunc   0x000004 /* Contains one or more functions of any kind */
-#define EP_FixedCol  0x000008 /* TK_Column with a known fixed value */
-#define EP_Agg       0x000010 /* Contains one or more aggregate functions */
-#define EP_VarSelect 0x000020 /* pSelect is correlated, not constant */
-#define EP_DblQuoted 0x000040 /* token.z was originally in "..." */
-#define EP_InfixFunc 0x000080 /* True for an infix function: LIKE, GLOB, etc */
-#define EP_Collate   0x000100 /* Tree contains a TK_COLLATE operator */
-#define EP_Generic   0x000200 /* Ignore COLLATE or affinity on this tree */
-#define EP_IntValue  0x000400 /* Integer value contained in u.iValue */
-#define EP_xIsSelect 0x000800 /* x.pSelect is valid (otherwise x.pList is) */
-#define EP_Skip      0x001000 /* COLLATE, AS, or UNLIKELY */
-#define EP_Reduced   0x002000 /* Expr struct EXPR_REDUCEDSIZE bytes only */
-#define EP_TokenOnly 0x004000 /* Expr struct EXPR_TOKENONLYSIZE bytes only */
-#define EP_Win       0x008000 /* Contains window functions */
-#define EP_MemToken  0x010000 /* Need to sqlite3DbFree() Expr.zToken */
-#define EP_NoReduce  0x020000 /* Cannot EXPRDUP_REDUCE this Expr */
-#define EP_Unlikely  0x040000 /* unlikely() or likelihood() function */
-#define EP_ConstFunc 0x080000 /* A SQLITE_FUNC_CONSTANT or _SLOCHNG function */
-#define EP_CanBeNull 0x100000 /* Can be null despite NOT NULL constraint */
-#define EP_Subquery  0x200000 /* Tree contains a TK_SELECT operator */
-#define EP_Alias     0x400000 /* Is an alias for a result set column */
-#define EP_Leaf      0x800000 /* Expr.pLeft, .pRight, .u.pSelect all NULL */
-#define EP_WinFunc  0x1000000 /* TK_FUNCTION with Expr.y.pWin set */
-#define EP_Subrtn   0x2000000 /* Uses Expr.y.sub. TK_IN, _SELECT, or _EXISTS */
-#define EP_Quoted   0x4000000 /* TK_ID was originally quoted */
-#define EP_Static   0x8000000 /* Held in memory not obtained from malloc() */
+#define EP_FromJoin   0x000001 /* Originates in ON/USING clause of outer join */
+#define EP_Distinct   0x000002 /* Aggregate function with DISTINCT keyword */
+#define EP_HasFunc    0x000004 /* Contains one or more functions of any kind */
+#define EP_FixedCol   0x000008 /* TK_Column with a known fixed value */
+#define EP_Agg        0x000010 /* Contains one or more aggregate functions */
+#define EP_VarSelect  0x000020 /* pSelect is correlated, not constant */
+#define EP_DblQuoted  0x000040 /* token.z was originally in "..." */
+#define EP_InfixFunc  0x000080 /* True for an infix function: LIKE, GLOB, etc */
+#define EP_Collate    0x000100 /* Tree contains a TK_COLLATE operator */
+#define EP_Commuted   0x000200 /* Comparison operator has been commuted */
+#define EP_IntValue   0x000400 /* Integer value contained in u.iValue */
+#define EP_xIsSelect  0x000800 /* x.pSelect is valid (otherwise x.pList is) */
+#define EP_Skip       0x001000 /* Operator does not contribute to affinity */
+#define EP_Reduced    0x002000 /* Expr struct EXPR_REDUCEDSIZE bytes only */
+#define EP_TokenOnly  0x004000 /* Expr struct EXPR_TOKENONLYSIZE bytes only */
+#define EP_Win        0x008000 /* Contains window functions */
+#define EP_MemToken   0x010000 /* Need to sqlite3DbFree() Expr.zToken */
+#define EP_NoReduce   0x020000 /* Cannot EXPRDUP_REDUCE this Expr */
+#define EP_Unlikely   0x040000 /* unlikely() or likelihood() function */
+#define EP_ConstFunc  0x080000 /* A SQLITE_FUNC_CONSTANT or _SLOCHNG function */
+#define EP_CanBeNull  0x100000 /* Can be null despite NOT NULL constraint */
+#define EP_Subquery   0x200000 /* Tree contains a TK_SELECT operator */
+#define EP_Alias      0x400000 /* Is an alias for a result set column */
+#define EP_Leaf       0x800000 /* Expr.pLeft, .pRight, .u.pSelect all NULL */
+#define EP_WinFunc   0x1000000 /* TK_FUNCTION with Expr.y.pWin set */
+#define EP_Subrtn    0x2000000 /* Uses Expr.y.sub. TK_IN, _SELECT, or _EXISTS */
+#define EP_Quoted    0x4000000 /* TK_ID was originally quoted */
+#define EP_Static    0x8000000 /* Held in memory not obtained from malloc() */
+#define EP_IsTrue   0x10000000 /* Always has boolean value of TRUE */
+#define EP_IsFalse  0x20000000 /* Always has boolean value of FALSE */
+#define EP_FromDDL  0x40000000 /* Originates from sqlite_master */
 
 /*
 ** The EP_Propagate mask is a set of properties that automatically propagate
@@ -17468,6 +17969,8 @@ struct Expr {
 #define ExprHasAllProperty(E,P)  (((E)->flags&(P))==(P))
 #define ExprSetProperty(E,P)     (E)->flags|=(P)
 #define ExprClearProperty(E,P)   (E)->flags&=~(P)
+#define ExprAlwaysTrue(E)   (((E)->flags&(EP_FromJoin|EP_IsTrue))==EP_IsTrue)
+#define ExprAlwaysFalse(E)  (((E)->flags&(EP_FromJoin|EP_IsFalse))==EP_IsFalse)
 
 /* The ExprSetVVAProperty() macro is used for Verification, Validation,
 ** and Accreditation only.  It works like ExprSetProperty() during VVA
@@ -17494,6 +17997,18 @@ struct Expr {
 */
 #define EXPRDUP_REDUCE         0x0001  /* Used reduced-size Expr nodes */
 
+/*
+** True if the expression passed as an argument was a function with
+** an OVER() clause (a window function).
+*/
+#ifdef SQLITE_OMIT_WINDOWFUNC
+# define IsWindowFunc(p) 0
+#else
+# define IsWindowFunc(p) ( \
+    ExprHasProperty((p), EP_WinFunc) && p->y.pWin->eFrmType!=TK_FILTER \
+ )
+#endif
+
 /*
 ** A list of expressions.  Each expression may optionally have a
 ** name.  An expr/name combination can be used in several ways, such
@@ -17502,25 +18017,31 @@ struct Expr {
 ** also be used as the argument to a function, in which case the a.zName
 ** field is not used.
 **
-** By default the Expr.zSpan field holds a human-readable description of
-** the expression that is used in the generation of error messages and
-** column labels.  In this case, Expr.zSpan is typically the text of a
-** column expression as it exists in a SELECT statement.  However, if
-** the bSpanIsTab flag is set, then zSpan is overloaded to mean the name
-** of the result column in the form: DATABASE.TABLE.COLUMN.  This later
-** form is used for name resolution with nested FROM clauses.
+** In order to try to keep memory usage down, the Expr.a.zEName field
+** is used for multiple purposes:
+**
+**     eEName          Usage
+**    ----------       -------------------------
+**    ENAME_NAME       (1) the AS of result set column
+**                     (2) COLUMN= of an UPDATE
+**
+**    ENAME_TAB        DB.TABLE.NAME used to resolve names
+**                     of subqueries
+**
+**    ENAME_SPAN       Text of the original result set
+**                     expression.
 */
 struct ExprList {
   int nExpr;             /* Number of expressions on the list */
   struct ExprList_item { /* For each expression in the list */
     Expr *pExpr;            /* The parse tree for this expression */
-    char *zName;            /* Token associated with this expression */
-    char *zSpan;            /* Original text of the expression */
-    u8 sortOrder;           /* 1 for DESC or 0 for ASC */
+    char *zEName;           /* Token associated with this expression */
+    u8 sortFlags;           /* Mask of KEYINFO_ORDER_* flags */
+    unsigned eEName :2;     /* Meaning of zEName */
     unsigned done :1;       /* A flag to indicate when processing is finished */
-    unsigned bSpanIsTab :1; /* zSpan holds DB.TABLE.COLUMN */
     unsigned reusable :1;   /* Constant expression is reusable */
     unsigned bSorterRef :1; /* Defer evaluation until after sorting */
+    unsigned bNulls: 1;     /* True if explicit "NULLS FIRST/LAST" */
     union {
       struct {
         u16 iOrderByCol;      /* For ORDER BY, column number in result set */
@@ -17531,6 +18052,13 @@ struct ExprList {
   } a[1];                  /* One slot for each expression in the list */
 };
 
+/*
+** Allowed values for Expr.a.eEName
+*/
+#define ENAME_NAME  0       /* The AS clause of a result set */
+#define ENAME_SPAN  1       /* Complete text of the result set expression */
+#define ENAME_TAB   2       /* "DB.TABLE.NAME" for the result set */
+
 /*
 ** An instance of this structure can hold a simple list of identifiers,
 ** such as the list "a,b,c" in the following statements:
@@ -17594,6 +18122,7 @@ struct SrcList {
       unsigned isCorrelated :1;  /* True if sub-query is correlated */
       unsigned viaCoroutine :1;  /* Implemented as a co-routine */
       unsigned isRecursive :1;   /* True for recursive reference in WITH */
+      unsigned fromDDL :1;       /* Comes from sqlite_master */
     } fg;
     int iCursor;      /* The VDBE cursor number used to access this table */
     Expr *pOn;        /* The ON clause of a join */
@@ -17684,7 +18213,7 @@ struct NameContext {
   NameContext *pNext;  /* Next outer name context.  NULL for outermost */
   int nRef;            /* Number of names resolved by this context */
   int nErr;            /* Number of errors encountered while resolving names */
-  u16 ncFlags;         /* Zero or more NC_* flags defined below */
+  int ncFlags;         /* Zero or more NC_* flags defined below */
   Select *pWinSelect;  /* SELECT statement for any window functions */
 };
 
@@ -17697,20 +18226,24 @@ struct NameContext {
 **    NC_HasWin    == EP_Win
 **
 */
-#define NC_AllowAgg  0x0001  /* Aggregate functions are allowed here */
-#define NC_PartIdx   0x0002  /* True if resolving a partial index WHERE */
-#define NC_IsCheck   0x0004  /* True if resolving names in a CHECK constraint */
-#define NC_InAggFunc 0x0008  /* True if analyzing arguments to an agg func */
-#define NC_HasAgg    0x0010  /* One or more aggregate functions seen */
-#define NC_IdxExpr   0x0020  /* True if resolving columns of CREATE INDEX */
-#define NC_VarSelect 0x0040  /* A correlated subquery has been seen */
-#define NC_UEList    0x0080  /* True if uNC.pEList is used */
-#define NC_UAggInfo  0x0100  /* True if uNC.pAggInfo is used */
-#define NC_UUpsert   0x0200  /* True if uNC.pUpsert is used */
-#define NC_MinMaxAgg 0x1000  /* min/max aggregates seen.  See note above */
-#define NC_Complex   0x2000  /* True if a function or subquery seen */
-#define NC_AllowWin  0x4000  /* Window functions are allowed here */
-#define NC_HasWin    0x8000  /* One or more window functions seen */
+#define NC_AllowAgg  0x00001  /* Aggregate functions are allowed here */
+#define NC_PartIdx   0x00002  /* True if resolving a partial index WHERE */
+#define NC_IsCheck   0x00004  /* True if resolving a CHECK constraint */
+#define NC_GenCol    0x00008  /* True for a GENERATED ALWAYS AS clause */
+#define NC_HasAgg    0x00010  /* One or more aggregate functions seen */
+#define NC_IdxExpr   0x00020  /* True if resolving columns of CREATE INDEX */
+#define NC_SelfRef   0x0002e  /* Combo: PartIdx, isCheck, GenCol, and IdxExpr */
+#define NC_VarSelect 0x00040  /* A correlated subquery has been seen */
+#define NC_UEList    0x00080  /* True if uNC.pEList is used */
+#define NC_UAggInfo  0x00100  /* True if uNC.pAggInfo is used */
+#define NC_UUpsert   0x00200  /* True if uNC.pUpsert is used */
+#define NC_MinMaxAgg 0x01000  /* min/max aggregates seen.  See note above */
+#define NC_Complex   0x02000  /* True if a function or subquery seen */
+#define NC_AllowWin  0x04000  /* Window functions are allowed here */
+#define NC_HasWin    0x08000  /* One or more window functions seen */
+#define NC_IsDDL     0x10000  /* Resolving names in a CREATE statement */
+#define NC_InAggFunc 0x20000  /* True if analyzing arguments to an agg func */
+#define NC_FromDDL   0x40000  /* SQL text comes from sqlite_master */
 
 /*
 ** An instance of the following object describes a single ON CONFLICT
@@ -17760,13 +18293,13 @@ struct Upsert {
 ** sequences for the ORDER BY clause.
 */
 struct Select {
-  ExprList *pEList;      /* The fields of the result */
   u8 op;                 /* One of: TK_UNION TK_ALL TK_INTERSECT TK_EXCEPT */
   LogEst nSelectRow;     /* Estimated number of result rows */
   u32 selFlags;          /* Various SF_* values */
   int iLimit, iOffset;   /* Memory registers holding LIMIT & OFFSET counters */
   u32 selId;             /* Unique identifier number for this SELECT */
   int addrOpenEphm[2];   /* OP_OpenEphem opcodes related to this select */
+  ExprList *pEList;      /* The fields of the result */
   SrcList *pSrc;         /* The FROM clause */
   Expr *pWhere;          /* The WHERE clause */
   ExprList *pGroupBy;    /* The GROUP BY clause */
@@ -17791,25 +18324,28 @@ struct Select {
 **     SF_MinMaxAgg  == NC_MinMaxAgg     == SQLITE_FUNC_MINMAX
 **     SF_FixedLimit == WHERE_USE_LIMIT
 */
-#define SF_Distinct       0x00001  /* Output should be DISTINCT */
-#define SF_All            0x00002  /* Includes the ALL keyword */
-#define SF_Resolved       0x00004  /* Identifiers have been resolved */
-#define SF_Aggregate      0x00008  /* Contains agg functions or a GROUP BY */
-#define SF_HasAgg         0x00010  /* Contains aggregate functions */
-#define SF_UsesEphemeral  0x00020  /* Uses the OpenEphemeral opcode */
-#define SF_Expanded       0x00040  /* sqlite3SelectExpand() called on this */
-#define SF_HasTypeInfo    0x00080  /* FROM subqueries have Table metadata */
-#define SF_Compound       0x00100  /* Part of a compound query */
-#define SF_Values         0x00200  /* Synthesized from VALUES clause */
-#define SF_MultiValue     0x00400  /* Single VALUES term with multiple rows */
-#define SF_NestedFrom     0x00800  /* Part of a parenthesized FROM clause */
-#define SF_MinMaxAgg      0x01000  /* Aggregate containing min() or max() */
-#define SF_Recursive      0x02000  /* The recursive part of a recursive CTE */
-#define SF_FixedLimit     0x04000  /* nSelectRow set by a constant LIMIT */
-#define SF_MaybeConvert   0x08000  /* Need convertCompoundSelectToSubquery() */
-#define SF_Converted      0x10000  /* By convertCompoundSelectToSubquery() */
-#define SF_IncludeHidden  0x20000  /* Include hidden columns in output */
-#define SF_ComplexResult  0x40000  /* Result contains subquery or function */
+#define SF_Distinct      0x0000001 /* Output should be DISTINCT */
+#define SF_All           0x0000002 /* Includes the ALL keyword */
+#define SF_Resolved      0x0000004 /* Identifiers have been resolved */
+#define SF_Aggregate     0x0000008 /* Contains agg functions or a GROUP BY */
+#define SF_HasAgg        0x0000010 /* Contains aggregate functions */
+#define SF_UsesEphemeral 0x0000020 /* Uses the OpenEphemeral opcode */
+#define SF_Expanded      0x0000040 /* sqlite3SelectExpand() called on this */
+#define SF_HasTypeInfo   0x0000080 /* FROM subqueries have Table metadata */
+#define SF_Compound      0x0000100 /* Part of a compound query */
+#define SF_Values        0x0000200 /* Synthesized from VALUES clause */
+#define SF_MultiValue    0x0000400 /* Single VALUES term with multiple rows */
+#define SF_NestedFrom    0x0000800 /* Part of a parenthesized FROM clause */
+#define SF_MinMaxAgg     0x0001000 /* Aggregate containing min() or max() */
+#define SF_Recursive     0x0002000 /* The recursive part of a recursive CTE */
+#define SF_FixedLimit    0x0004000 /* nSelectRow set by a constant LIMIT */
+#define SF_MaybeConvert  0x0008000 /* Need convertCompoundSelectToSubquery() */
+#define SF_Converted     0x0010000 /* By convertCompoundSelectToSubquery() */
+#define SF_IncludeHidden 0x0020000 /* Include hidden columns in output */
+#define SF_ComplexResult 0x0040000 /* Result contains subquery or function */
+#define SF_WhereBegin    0x0080000 /* Really a WhereBegin() call.  Debug Only */
+#define SF_WinRewrite    0x0100000 /* Window function rewrite accomplished */
+#define SF_View          0x0200000 /* SELECT statement is a view */
 
 /*
 ** The results of a SELECT can be distributed in several ways, as defined
@@ -18089,8 +18625,8 @@ struct Parse {
 
 #define PARSE_MODE_NORMAL        0
 #define PARSE_MODE_DECLARE_VTAB  1
-#define PARSE_MODE_RENAME_COLUMN 2
-#define PARSE_MODE_RENAME_TABLE  3
+#define PARSE_MODE_RENAME        2
+#define PARSE_MODE_UNMAP         3
 
 /*
 ** Sizes and pointers of various parts of the Parse object.
@@ -18112,7 +18648,7 @@ struct Parse {
 #if defined(SQLITE_OMIT_ALTERTABLE)
   #define IN_RENAME_OBJECT 0
 #else
-  #define IN_RENAME_OBJECT (pParse->eParseMode>=PARSE_MODE_RENAME_COLUMN)
+  #define IN_RENAME_OBJECT (pParse->eParseMode>=PARSE_MODE_RENAME)
 #endif
 
 #if defined(SQLITE_OMIT_VIRTUALTABLE) && defined(SQLITE_OMIT_ALTERTABLE)
@@ -18263,7 +18799,7 @@ typedef struct DbFixer DbFixer;
 struct DbFixer {
   Parse *pParse;      /* The parsing context.  Error messages written here */
   Schema *pSchema;    /* Fix items to this schema */
-  int bVarOnly;       /* Check for variable references only */
+  u8 bTemp;           /* True for TEMP schema entries */
   const char *zDb;    /* Make sure all objects are contained in this database */
   const char *zType;  /* Type of the container - used for error messages */
   const Token *pName; /* Name of the container - used for error messages */
@@ -18314,11 +18850,12 @@ typedef struct {
 */
 struct Sqlite3Config {
   int bMemstat;                     /* True to enable memory status */
-  int bCoreMutex;                   /* True to enable core mutexing */
-  int bFullMutex;                   /* True to enable full mutexing */
-  int bOpenUri;                     /* True to interpret filenames as URIs */
-  int bUseCis;                      /* Use covering indices for full-scans */
-  int bSmallMalloc;                 /* Avoid large memory allocations if true */
+  u8 bCoreMutex;                    /* True to enable core mutexing */
+  u8 bFullMutex;                    /* True to enable full mutexing */
+  u8 bOpenUri;                      /* True to interpret filenames as URIs */
+  u8 bUseCis;                       /* Use covering indices for full-scans */
+  u8 bSmallMalloc;                  /* Avoid large memory allocations if true */
+  u8 bExtraSchemaChecks;            /* Verify type,name,tbl_name in schema */
   int mxStrlen;                     /* Maximum string length */
   int neverCorrupt;                 /* Database is always well-formed */
   int szLookaside;                  /* Default lookaside buffer size */
@@ -18367,9 +18904,9 @@ struct Sqlite3Config {
   int (*xTestCallback)(int);        /* Invoked by sqlite3FaultSim() */
 #endif
   int bLocaltimeFault;              /* True to fail localtime() calls */
-  int bInternalFunctions;           /* Internal SQL functions are visible */
   int iOnceResetThreshold;          /* When to reset OP_Once counters */
   u32 szSorterRef;                  /* Min size in bytes to use sorter-refs */
+  unsigned int iPrngSeed;           /* Alternative fixed seed for the PRNG */
 };
 
 /*
@@ -18399,7 +18936,7 @@ struct Walker {
   int (*xSelectCallback)(Walker*,Select*);  /* Callback for SELECTs */
   void (*xSelectCallback2)(Walker*,Select*);/* Second callback for SELECTs */
   int walkerDepth;                          /* Number of subqueries */
-  u8 eCode;                                 /* A small processing code */
+  u16 eCode;                                /* A small processing code */
   union {                                   /* Extra data for callback */
     NameContext *pNC;                         /* Naming context */
     int n;                                    /* A counter */
@@ -18415,6 +18952,7 @@ struct Walker {
     struct WindowRewrite *pRewrite;           /* Window rewrite context */
     struct WhereConst *pConst;                /* WHERE clause constants */
     struct RenameCtx *pRename;                /* RENAME COLUMN context */
+    struct Table *pTab;                       /* Table of generated column */
   } u;
 };
 
@@ -18466,10 +19004,11 @@ struct TreeView {
 #endif /* SQLITE_DEBUG */
 
 /*
-** This object is used in various ways, all related to window functions
+** This object is used in various ways, most (but not all) related to window
+** functions.
 **
 **   (1) A single instance of this structure is attached to the
-**       the Expr.pWin field for each window function in an expression tree.
+**       the Expr.y.pWin field for each window function in an expression tree.
 **       This object holds the information contained in the OVER clause,
 **       plus additional fields used during code generation.
 **
@@ -18480,6 +19019,10 @@ struct TreeView {
 **   (3) The terms of the WINDOW clause of a SELECT are instances of this
 **       object on a linked list attached to Select.pWinDefn.
 **
+**   (4) For an aggregate function with a FILTER clause, an instance
+**       of this object is stored in Expr.y.pWin with eFrmType set to
+**       TK_FILTER. In this case the only field used is Window.pFilter.
+**
 ** The uses (1) and (2) are really the same Window object that just happens
 ** to be accessible in two different ways.  Use case (3) are separate objects.
 */
@@ -18495,12 +19038,13 @@ struct Window {
   u8 eExclude;            /* TK_NO, TK_CURRENT, TK_TIES, TK_GROUP, or 0 */
   Expr *pStart;           /* Expression for "<expr> PRECEDING" */
   Expr *pEnd;             /* Expression for "<expr> FOLLOWING" */
+  Window **ppThis;        /* Pointer to this object in Select.pWin list */
   Window *pNextWin;       /* Next window function belonging to this SELECT */
   Expr *pFilter;          /* The FILTER expression */
   FuncDef *pFunc;         /* The function */
   int iEphCsr;            /* Partition buffer or Peer buffer */
-  int regAccum;
-  int regResult;
+  int regAccum;           /* Accumulator */
+  int regResult;          /* Interim result */
   int csrApp;             /* Function cursor (used by min/max) */
   int regApp;             /* Function register (also used by min/max) */
   int regPart;            /* Array of registers for PARTITION BY values */
@@ -18510,15 +19054,19 @@ struct Window {
   int regOne;             /* Register containing constant value 1 */
   int regStartRowid;
   int regEndRowid;
+  u8 bExprArgs;           /* Defer evaluation of window function arguments
+                          ** due to the SQLITE_SUBTYPE flag */
 };
 
 #ifndef SQLITE_OMIT_WINDOWFUNC
 SQLITE_PRIVATE void sqlite3WindowDelete(sqlite3*, Window*);
+SQLITE_PRIVATE void sqlite3WindowUnlinkFromSelect(Window*);
 SQLITE_PRIVATE void sqlite3WindowListDelete(sqlite3 *db, Window *p);
 SQLITE_PRIVATE Window *sqlite3WindowAlloc(Parse*, int, int, Expr*, int , Expr*, u8);
 SQLITE_PRIVATE void sqlite3WindowAttach(Parse*, Expr*, Window*);
-SQLITE_PRIVATE int sqlite3WindowCompare(Parse*, Window*, Window*);
-SQLITE_PRIVATE void sqlite3WindowCodeInit(Parse*, Window*);
+SQLITE_PRIVATE void sqlite3WindowLink(Select *pSel, Window *pWin);
+SQLITE_PRIVATE int sqlite3WindowCompare(Parse*, Window*, Window*, int);
+SQLITE_PRIVATE void sqlite3WindowCodeInit(Parse*, Select*);
 SQLITE_PRIVATE void sqlite3WindowCodeStep(Parse*, Select*, WhereInfo*, int, int);
 SQLITE_PRIVATE int sqlite3WindowRewrite(Parse*, Select*);
 SQLITE_PRIVATE int sqlite3ExpandSubquery(Parse*, struct SrcList_item*);
@@ -18717,8 +19265,12 @@ SQLITE_PRIVATE void sqlite3MutexWarnOnContention(sqlite3_mutex*);
 #endif
 
 #ifndef SQLITE_OMIT_FLOATING_POINT
+# define EXP754 (((u64)0x7ff)<<52)
+# define MAN754 ((((u64)1)<<52)-1)
+# define IsNaN(X) (((X)&EXP754)==EXP754 && ((X)&MAN754)!=0)
 SQLITE_PRIVATE   int sqlite3IsNaN(double);
 #else
+# define IsNaN(X)         0
 # define sqlite3IsNaN(X)  0
 #endif
 
@@ -18777,13 +19329,16 @@ SQLITE_PRIVATE Expr *sqlite3Expr(sqlite3*,int,const char*);
 SQLITE_PRIVATE void sqlite3ExprAttachSubtrees(sqlite3*,Expr*,Expr*,Expr*);
 SQLITE_PRIVATE Expr *sqlite3PExpr(Parse*, int, Expr*, Expr*);
 SQLITE_PRIVATE void sqlite3PExprAddSelect(Parse*, Expr*, Select*);
-SQLITE_PRIVATE Expr *sqlite3ExprAnd(sqlite3*,Expr*, Expr*);
+SQLITE_PRIVATE Expr *sqlite3ExprAnd(Parse*,Expr*, Expr*);
+SQLITE_PRIVATE Expr *sqlite3ExprSimplifiedAndOr(Expr*);
 SQLITE_PRIVATE Expr *sqlite3ExprFunction(Parse*,ExprList*, Token*, int);
+SQLITE_PRIVATE void sqlite3ExprFunctionUsable(Parse*,Expr*,FuncDef*);
 SQLITE_PRIVATE void sqlite3ExprAssignVarNumber(Parse*, Expr*, u32);
 SQLITE_PRIVATE void sqlite3ExprDelete(sqlite3*, Expr*);
+SQLITE_PRIVATE void sqlite3ExprUnmapAndDelete(Parse*, Expr*);
 SQLITE_PRIVATE ExprList *sqlite3ExprListAppend(Parse*,ExprList*,Expr*);
 SQLITE_PRIVATE ExprList *sqlite3ExprListAppendVector(Parse*,ExprList*,IdList*,Expr*);
-SQLITE_PRIVATE void sqlite3ExprListSetSortOrder(ExprList*,int);
+SQLITE_PRIVATE void sqlite3ExprListSetSortOrder(ExprList*,int,int);
 SQLITE_PRIVATE void sqlite3ExprListSetName(Parse*,ExprList*,Token*,int);
 SQLITE_PRIVATE void sqlite3ExprListSetSpan(Parse*,ExprList*,const char*,const char*);
 SQLITE_PRIVATE void sqlite3ExprListDelete(sqlite3*, ExprList*);
@@ -18802,11 +19357,18 @@ SQLITE_PRIVATE void sqlite3CollapseDatabaseArray(sqlite3*);
 SQLITE_PRIVATE void sqlite3CommitInternalChanges(sqlite3*);
 SQLITE_PRIVATE void sqlite3DeleteColumnNames(sqlite3*,Table*);
 SQLITE_PRIVATE int sqlite3ColumnsFromExprList(Parse*,ExprList*,i16*,Column**);
-SQLITE_PRIVATE void sqlite3SelectAddColumnTypeAndCollation(Parse*,Table*,Select*);
-SQLITE_PRIVATE Table *sqlite3ResultSetOfSelect(Parse*,Select*);
+SQLITE_PRIVATE void sqlite3SelectAddColumnTypeAndCollation(Parse*,Table*,Select*,char);
+SQLITE_PRIVATE Table *sqlite3ResultSetOfSelect(Parse*,Select*,char);
 SQLITE_PRIVATE void sqlite3OpenMasterTable(Parse *, int);
 SQLITE_PRIVATE Index *sqlite3PrimaryKeyIndex(Table*);
-SQLITE_PRIVATE i16 sqlite3ColumnOfIndex(Index*, i16);
+SQLITE_PRIVATE i16 sqlite3TableColumnToIndex(Index*, i16);
+#ifdef SQLITE_OMIT_GENERATED_COLUMNS
+# define sqlite3TableColumnToStorage(T,X) (X)  /* No-op pass-through */
+# define sqlite3StorageColumnToTable(T,X) (X)  /* No-op pass-through */
+#else
+SQLITE_PRIVATE   i16 sqlite3TableColumnToStorage(Table*, i16);
+SQLITE_PRIVATE   i16 sqlite3StorageColumnToTable(Table*, i16);
+#endif
 SQLITE_PRIVATE void sqlite3StartTable(Parse*,Token*,Token*,int,int,int,int);
 #if SQLITE_ENABLE_HIDDEN_COLUMNS
 SQLITE_PRIVATE   void sqlite3ColumnPropertiesFromName(Table*, Column*);
@@ -18819,6 +19381,7 @@ SQLITE_PRIVATE void sqlite3AddPrimaryKey(Parse*, ExprList*, int, int, int);
 SQLITE_PRIVATE void sqlite3AddCheckConstraint(Parse*, Expr*);
 SQLITE_PRIVATE void sqlite3AddDefaultValue(Parse*,Expr*,const char*,const char*);
 SQLITE_PRIVATE void sqlite3AddCollateType(Parse*, Token*);
+SQLITE_PRIVATE void sqlite3AddGenerated(Parse*,Expr*,Token*);
 SQLITE_PRIVATE void sqlite3EndTable(Parse*,Token*,Token*,u8,Select*);
 SQLITE_PRIVATE int sqlite3ParseUri(const char*,const char*,unsigned int*,
                     sqlite3_vfs**,char**,char **);
@@ -18876,6 +19439,9 @@ SQLITE_PRIVATE   void sqlite3AutoincrementEnd(Parse *pParse);
 # define sqlite3AutoincrementEnd(X)
 #endif
 SQLITE_PRIVATE void sqlite3Insert(Parse*, SrcList*, Select*, IdList*, int, Upsert*);
+#ifndef SQLITE_OMIT_GENERATED_COLUMNS
+SQLITE_PRIVATE   void sqlite3ComputeGeneratedColumns(Parse*, int, Table*);
+#endif
 SQLITE_PRIVATE void *sqlite3ArrayAllocate(sqlite3*,void*,int,int*,int*);
 SQLITE_PRIVATE IdList *sqlite3IdListAppend(Parse*, IdList*, Token*);
 SQLITE_PRIVATE int sqlite3IdListIndex(IdList*,const char*);
@@ -18898,6 +19464,7 @@ SQLITE_PRIVATE int sqlite3Select(Parse*, Select*, SelectDest*);
 SQLITE_PRIVATE Select *sqlite3SelectNew(Parse*,ExprList*,SrcList*,Expr*,ExprList*,
                          Expr*,ExprList*,u32,Expr*);
 SQLITE_PRIVATE void sqlite3SelectDelete(sqlite3*, Select*);
+SQLITE_PRIVATE void sqlite3SelectReset(Parse*, Select*);
 SQLITE_PRIVATE Table *sqlite3SrcListLookup(Parse*, SrcList*);
 SQLITE_PRIVATE int sqlite3IsReadOnly(Parse*, Table*, int);
 SQLITE_PRIVATE void sqlite3OpenTable(Parse*, int iCur, int iDb, Table*, int);
@@ -18920,17 +19487,20 @@ SQLITE_PRIVATE int sqlite3WhereOkOnePass(WhereInfo*, int*);
 #define ONEPASS_OFF      0        /* Use of ONEPASS not allowed */
 #define ONEPASS_SINGLE   1        /* ONEPASS valid for a single row update */
 #define ONEPASS_MULTI    2        /* ONEPASS is valid for multiple rows */
+SQLITE_PRIVATE int sqlite3WhereUsesDeferredSeek(WhereInfo*);
 SQLITE_PRIVATE void sqlite3ExprCodeLoadIndexColumn(Parse*, Index*, int, int, int);
 SQLITE_PRIVATE int sqlite3ExprCodeGetColumn(Parse*, Table*, int, int, int, u8);
 SQLITE_PRIVATE void sqlite3ExprCodeGetColumnOfTable(Vdbe*, Table*, int, int, int);
 SQLITE_PRIVATE void sqlite3ExprCodeMove(Parse*, int, int, int);
 SQLITE_PRIVATE void sqlite3ExprCode(Parse*, Expr*, int);
+#ifndef SQLITE_OMIT_GENERATED_COLUMNS
+SQLITE_PRIVATE void sqlite3ExprCodeGeneratedColumn(Parse*, Column*, int);
+#endif
 SQLITE_PRIVATE void sqlite3ExprCodeCopy(Parse*, Expr*, int);
 SQLITE_PRIVATE void sqlite3ExprCodeFactorable(Parse*, Expr*, int);
 SQLITE_PRIVATE int sqlite3ExprCodeAtInit(Parse*, Expr*, int);
 SQLITE_PRIVATE int sqlite3ExprCodeTemp(Parse*, Expr*, int*);
 SQLITE_PRIVATE int sqlite3ExprCodeTarget(Parse*, Expr*, int);
-SQLITE_PRIVATE void sqlite3ExprCodeAndCache(Parse*, Expr*, int);
 SQLITE_PRIVATE int sqlite3ExprCodeExprList(Parse*, ExprList*, int, int, u8);
 #define SQLITE_ECEL_DUP      0x01  /* Deep, not shallow copies */
 #define SQLITE_ECEL_FACTOR   0x02  /* Factor out constant terms */
@@ -18972,6 +19542,7 @@ SQLITE_PRIVATE void sqlite3EndTransaction(Parse*,int);
 SQLITE_PRIVATE void sqlite3Savepoint(Parse*, int, Token*);
 SQLITE_PRIVATE void sqlite3CloseSavepoints(sqlite3 *);
 SQLITE_PRIVATE void sqlite3LeaveMutexAndCloseZombie(sqlite3*);
+SQLITE_PRIVATE u32 sqlite3IsTrueOrFalse(const char*);
 SQLITE_PRIVATE int sqlite3ExprIdToTrueFalse(Expr*);
 SQLITE_PRIVATE int sqlite3ExprTruthValue(const Expr*);
 SQLITE_PRIVATE int sqlite3ExprIsConstant(Expr*);
@@ -19067,6 +19638,7 @@ SQLITE_PRIVATE   u32 sqlite3TriggerColmask(Parse*,Trigger*,ExprList*,int,int,Tab
 #endif
 
 SQLITE_PRIVATE int sqlite3JoinType(Parse*, Token*, Token*, Token*);
+SQLITE_PRIVATE void sqlite3SetJoinExpr(Expr*,int);
 SQLITE_PRIVATE void sqlite3CreateForeignKey(Parse*, ExprList*, Token*, ExprList*, int);
 SQLITE_PRIVATE void sqlite3DeferForeignKey(Parse*, int);
 #ifndef SQLITE_OMIT_AUTHORIZATION
@@ -19089,6 +19661,7 @@ SQLITE_PRIVATE int sqlite3FixSelect(DbFixer*, Select*);
 SQLITE_PRIVATE int sqlite3FixExpr(DbFixer*, Expr*);
 SQLITE_PRIVATE int sqlite3FixExprList(DbFixer*, ExprList*);
 SQLITE_PRIVATE int sqlite3FixTriggerStep(DbFixer*, TriggerStep*);
+SQLITE_PRIVATE int sqlite3RealSameAsInt(double,sqlite3_int64);
 SQLITE_PRIVATE int sqlite3AtoF(const char *z, double*, int, u8);
 SQLITE_PRIVATE int sqlite3GetInt32(const char *, int*);
 SQLITE_PRIVATE int sqlite3Atoi(const char*);
@@ -19103,7 +19676,7 @@ SQLITE_PRIVATE LogEst sqlite3LogEstAdd(LogEst,LogEst);
 SQLITE_PRIVATE LogEst sqlite3LogEstFromDouble(double);
 #endif
 #if defined(SQLITE_ENABLE_STMT_SCANSTATUS) || \
-    defined(SQLITE_ENABLE_STAT3_OR_STAT4) || \
+    defined(SQLITE_ENABLE_STAT4) || \
     defined(SQLITE_EXPLAIN_ESTIMATED_ROWS)
 SQLITE_PRIVATE u64 sqlite3LogEstToInt(LogEst);
 #endif
@@ -19169,9 +19742,10 @@ SQLITE_PRIVATE int sqlite3ExprCollSeqMatch(Parse*,Expr*,Expr*);
 SQLITE_PRIVATE Expr *sqlite3ExprAddCollateToken(Parse *pParse, Expr*, const Token*, int);
 SQLITE_PRIVATE Expr *sqlite3ExprAddCollateString(Parse*,Expr*,const char*);
 SQLITE_PRIVATE Expr *sqlite3ExprSkipCollate(Expr*);
+SQLITE_PRIVATE Expr *sqlite3ExprSkipCollateAndLikely(Expr*);
 SQLITE_PRIVATE int sqlite3CheckCollSeq(Parse *, CollSeq *);
 SQLITE_PRIVATE int sqlite3WritableSchema(sqlite3*);
-SQLITE_PRIVATE int sqlite3CheckObjectName(Parse *, const char *);
+SQLITE_PRIVATE int sqlite3CheckObjectName(Parse*, const char*,const char*,const char*);
 SQLITE_PRIVATE void sqlite3VdbeSetChanges(sqlite3 *, int);
 SQLITE_PRIVATE int sqlite3AddInt64(i64*,i64);
 SQLITE_PRIVATE int sqlite3SubInt64(i64*,i64);
@@ -19190,6 +19764,9 @@ SQLITE_PRIVATE void sqlite3ValueSetStr(sqlite3_value*, int, const void *,u8,
                         void(*)(void*));
 SQLITE_PRIVATE void sqlite3ValueSetNull(sqlite3_value*);
 SQLITE_PRIVATE void sqlite3ValueFree(sqlite3_value*);
+#ifndef SQLITE_UNTESTABLE
+SQLITE_PRIVATE void sqlite3ResultIntReal(sqlite3_context*);
+#endif
 SQLITE_PRIVATE sqlite3_value *sqlite3ValueNew(sqlite3 *);
 #ifndef SQLITE_OMIT_UTF16
 SQLITE_PRIVATE char *sqlite3Utf16to8(sqlite3 *, const void*, int, u8);
@@ -19201,7 +19778,6 @@ SQLITE_PRIVATE const unsigned char sqlite3OpcodeProperty[];
 SQLITE_PRIVATE const char sqlite3StrBINARY[];
 SQLITE_PRIVATE const unsigned char sqlite3UpperToLower[];
 SQLITE_PRIVATE const unsigned char sqlite3CtypeMap[];
-SQLITE_PRIVATE const Token sqlite3IntTokens[];
 SQLITE_PRIVATE SQLITE_WSD struct Sqlite3Config sqlite3Config;
 SQLITE_PRIVATE FuncDefHash sqlite3BuiltinFunctions;
 #ifndef SQLITE_OMIT_WSD
@@ -19223,7 +19799,12 @@ SQLITE_PRIVATE void sqlite3CodeRhsOfIN(Parse*, Expr*, int);
 SQLITE_PRIVATE int sqlite3CodeSubselect(Parse*, Expr*);
 SQLITE_PRIVATE void sqlite3SelectPrep(Parse*, Select*, NameContext*);
 SQLITE_PRIVATE void sqlite3SelectWrongNumTermsError(Parse *pParse, Select *p);
-SQLITE_PRIVATE int sqlite3MatchSpanName(const char*, const char*, const char*, const char*);
+SQLITE_PRIVATE int sqlite3MatchEName(
+  const struct ExprList_item*,
+  const char*,
+  const char*,
+  const char*
+);
 SQLITE_PRIVATE int sqlite3ResolveExprNames(NameContext*, Expr*);
 SQLITE_PRIVATE int sqlite3ResolveExprListNames(NameContext*, ExprList*);
 SQLITE_PRIVATE void sqlite3ResolveSelectNames(Parse*, Select*, NameContext*);
@@ -19255,6 +19836,7 @@ SQLITE_PRIVATE void sqlite3KeyInfoUnref(KeyInfo*);
 SQLITE_PRIVATE KeyInfo *sqlite3KeyInfoRef(KeyInfo*);
 SQLITE_PRIVATE KeyInfo *sqlite3KeyInfoOfIndex(Parse*, Index*);
 SQLITE_PRIVATE KeyInfo *sqlite3KeyInfoFromExprList(Parse*, ExprList*, int, int);
+SQLITE_PRIVATE int sqlite3HasExplicitNulls(Parse*, ExprList*);
 
 #ifdef SQLITE_DEBUG
 SQLITE_PRIVATE int sqlite3KeyInfoIsWriteable(KeyInfo*);
@@ -19287,8 +19869,7 @@ SQLITE_PRIVATE int sqlite3ExprCheckIN(Parse*, Expr*);
 # define sqlite3ExprCheckIN(x,y) SQLITE_OK
 #endif
 
-#ifdef SQLITE_ENABLE_STAT3_OR_STAT4
-SQLITE_PRIVATE void sqlite3AnalyzeFunctions(void);
+#ifdef SQLITE_ENABLE_STAT4
 SQLITE_PRIVATE int sqlite3Stat4ProbeSetValue(
     Parse*,Index*,UnpackedRecord**,Expr*,int,int,int*);
 SQLITE_PRIVATE int sqlite3Stat4ValueFromExpr(Parse*, Expr*, u8, sqlite3_value**);
@@ -19335,6 +19916,7 @@ SQLITE_PRIVATE   int sqlite3Utf8To8(unsigned char*);
 #  define sqlite3VtabInSync(db) 0
 #  define sqlite3VtabLock(X)
 #  define sqlite3VtabUnlock(X)
+#  define sqlite3VtabModuleUnref(D,X)
 #  define sqlite3VtabUnlockList(X)
 #  define sqlite3VtabSavepoint(X, Y, Z) SQLITE_OK
 #  define sqlite3GetVTable(X,Y)  ((VTable*)0)
@@ -19346,6 +19928,7 @@ SQLITE_PRIVATE    int sqlite3VtabRollback(sqlite3 *db);
 SQLITE_PRIVATE    int sqlite3VtabCommit(sqlite3 *db);
 SQLITE_PRIVATE    void sqlite3VtabLock(VTable *);
 SQLITE_PRIVATE    void sqlite3VtabUnlock(VTable *);
+SQLITE_PRIVATE    void sqlite3VtabModuleUnref(sqlite3*,Module*);
 SQLITE_PRIVATE    void sqlite3VtabUnlockList(sqlite3*);
 SQLITE_PRIVATE    int sqlite3VtabSavepoint(sqlite3 *, int, int);
 SQLITE_PRIVATE    void sqlite3VtabImportErrmsg(Vdbe*, sqlite3_vtab*);
@@ -19359,6 +19942,12 @@ SQLITE_PRIVATE    Module *sqlite3VtabCreateModule(
    );
 #  define sqlite3VtabInSync(db) ((db)->nVTrans>0 && (db)->aVTrans==0)
 #endif
+SQLITE_PRIVATE int sqlite3ReadOnlyShadowTables(sqlite3 *db);
+#ifndef SQLITE_OMIT_VIRTUALTABLE
+SQLITE_PRIVATE   int sqlite3ShadowTableName(sqlite3 *db, const char *zName);
+#else
+# define sqlite3ShadowTableName(A,B) 0
+#endif
 SQLITE_PRIVATE int sqlite3VtabEponymousTableInit(Parse*,Module*);
 SQLITE_PRIVATE void sqlite3VtabEponymousTableClear(sqlite3*,Module*);
 SQLITE_PRIVATE void sqlite3VtabMakeWritable(Parse*,Table*);
@@ -19380,6 +19969,7 @@ SQLITE_PRIVATE char *sqlite3Normalize(Vdbe*, const char*);
 #endif
 SQLITE_PRIVATE int sqlite3Reprepare(Vdbe*);
 SQLITE_PRIVATE void sqlite3ExprListCheckLength(Parse*, ExprList*, const char*);
+SQLITE_PRIVATE CollSeq *sqlite3ExprCompareCollSeq(Parse*,Expr*);
 SQLITE_PRIVATE CollSeq *sqlite3BinaryCompareCollSeq(Parse *, Expr *, Expr *);
 SQLITE_PRIVATE int sqlite3TempInMemory(const sqlite3*);
 SQLITE_PRIVATE const char *sqlite3JournalModename(int);
@@ -19686,7 +20276,6 @@ SQLITE_PRIVATE const unsigned char sqlite3UpperToLower[] = {
 ** non-ASCII UTF character. Hence the test for whether or not a character is
 ** part of an identifier is 0x46.
 */
-#ifdef SQLITE_ASCII
 SQLITE_PRIVATE const unsigned char sqlite3CtypeMap[256] = {
   0x00, 0x00, 0x00, 0x00, 0x00, 0x00, 0x00, 0x00,  /* 00..07    ........ */
   0x00, 0x01, 0x01, 0x01, 0x01, 0x01, 0x00, 0x00,  /* 08..0f    ........ */
@@ -19724,7 +20313,6 @@ SQLITE_PRIVATE const unsigned char sqlite3CtypeMap[256] = {
   0x40, 0x40, 0x40, 0x40, 0x40, 0x40, 0x40, 0x40,  /* f0..f7    ........ */
   0x40, 0x40, 0x40, 0x40, 0x40, 0x40, 0x40, 0x40   /* f8..ff    ........ */
 };
-#endif
 
 /* EVIDENCE-OF: R-02982-34736 In order to maintain full backwards
 ** compatibility for legacy applications, the URI filename capability is
@@ -19752,8 +20340,15 @@ SQLITE_PRIVATE const unsigned char sqlite3CtypeMap[256] = {
 ** SQLITE_ALLOW_COVERING_INDEX_SCAN compile-time option, or is "on" if
 ** that compile-time option is omitted.
 */
-#ifndef SQLITE_ALLOW_COVERING_INDEX_SCAN
+#if !defined(SQLITE_ALLOW_COVERING_INDEX_SCAN)
 # define SQLITE_ALLOW_COVERING_INDEX_SCAN 1
+#else
+# if !SQLITE_ALLOW_COVERING_INDEX_SCAN 
+#   error "Compile-time disabling of covering index scan using the\
+ -DSQLITE_ALLOW_COVERING_INDEX_SCAN=0 option is deprecated.\
+ Contact SQLite developers if this is a problem for you, and\
+ delete this #error macro to continue with your build."
+# endif
 #endif
 
 /* The minimum PMA size is set to this value multiplied by the database
@@ -19782,9 +20377,18 @@ SQLITE_PRIVATE const unsigned char sqlite3CtypeMap[256] = {
 ** changed as start-time using sqlite3_config(SQLITE_CONFIG_LOOKASIDE)
 ** or at run-time for an individual database connection using
 ** sqlite3_db_config(db, SQLITE_DBCONFIG_LOOKASIDE);
+**
+** With the two-size-lookaside enhancement, less lookaside is required.
+** The default configuration of 1200,40 actually provides 30 1200-byte slots
+** and 93 128-byte slots, which is more lookaside than is available
+** using the older 1200,100 configuration without two-size-lookaside.
 */
 #ifndef SQLITE_DEFAULT_LOOKASIDE
-# define SQLITE_DEFAULT_LOOKASIDE 1200,100
+# ifdef SQLITE_OMIT_TWOSIZE_LOOKASIDE
+#   define SQLITE_DEFAULT_LOOKASIDE 1200,100  /* 120KB of memory */
+# else
+#   define SQLITE_DEFAULT_LOOKASIDE 1200,40   /* 48KB of memory */
+# endif
 #endif
 
 
@@ -19806,6 +20410,7 @@ SQLITE_PRIVATE SQLITE_WSD struct Sqlite3Config sqlite3Config = {
    SQLITE_USE_URI,            /* bOpenUri */
    SQLITE_ALLOW_COVERING_INDEX_SCAN,   /* bUseCis */
    0,                         /* bSmallMalloc */
+   1,                         /* bExtraSchemaChecks */
    0x7ffffffe,                /* mxStrlen */
    0,                         /* neverCorrupt */
    SQLITE_DEFAULT_LOOKASIDE,  /* szLookaside, nLookaside */
@@ -19849,9 +20454,9 @@ SQLITE_PRIVATE SQLITE_WSD struct Sqlite3Config sqlite3Config = {
    0,                         /* xTestCallback */
 #endif
    0,                         /* bLocaltimeFault */
-   0,                         /* bInternalFunctions */
    0x7ffffffe,                /* iOnceResetThreshold */
    SQLITE_DEFAULT_SORTERREF_SIZE,   /* szSorterRef */
+   0,                         /* iPrngSeed */
 };
 
 /*
@@ -19861,14 +20466,6 @@ SQLITE_PRIVATE SQLITE_WSD struct Sqlite3Config sqlite3Config = {
 */
 SQLITE_PRIVATE FuncDefHash sqlite3BuiltinFunctions;
 
-/*
-** Constant tokens for values 0 and 1.
-*/
-SQLITE_PRIVATE const Token sqlite3IntTokens[] = {
-   { "0", 1 },
-   { "1", 1 }
-};
-
 #ifdef VDBE_PROFILE
 /*
 ** The following performance counter can be used in place of
@@ -20180,12 +20777,12 @@ struct sqlite3_value {
 #define MEM_Int       0x0004   /* Value is an integer */
 #define MEM_Real      0x0008   /* Value is a real number */
 #define MEM_Blob      0x0010   /* Value is a BLOB */
-#define MEM_AffMask   0x001f   /* Mask of affinity bits */
-#define MEM_FromBind  0x0020   /* Value originates from sqlite3_bind() */
-/* Available          0x0040   */
+#define MEM_IntReal   0x0020   /* MEM_Int that stringifies like MEM_Real */
+#define MEM_AffMask   0x003f   /* Mask of affinity bits */
+#define MEM_FromBind  0x0040   /* Value originates from sqlite3_bind() */
 #define MEM_Undefined 0x0080   /* Value is undefined */
 #define MEM_Cleared   0x0100   /* NULL set by OP_Null, not from data */
-#define MEM_TypeMask  0xc1df   /* Mask of type bits */
+#define MEM_TypeMask  0xc1bf   /* Mask of type bits */
 
 
 /* Whenever Mem contains a valid string or blob representation, one of
@@ -20221,7 +20818,8 @@ struct sqlite3_value {
 ** True if Mem X is a NULL-nochng type.
 */
 #define MemNullNochng(X) \
-  ((X)->flags==(MEM_Null|MEM_Zero) && (X)->n==0 && (X)->u.nZero==0)
+  (((X)->flags&MEM_TypeMask)==(MEM_Null|MEM_Zero) \
+    && (X)->n==0 && (X)->u.nZero==0)
 
 /*
 ** Return true if a memory cell is not marked as invalid.  This macro
@@ -20417,11 +21015,11 @@ struct PreUpdate {
 SQLITE_PRIVATE void sqlite3VdbeError(Vdbe*, const char *, ...);
 SQLITE_PRIVATE void sqlite3VdbeFreeCursor(Vdbe *, VdbeCursor*);
 void sqliteVdbePopStack(Vdbe*,int);
+SQLITE_PRIVATE int SQLITE_NOINLINE sqlite3VdbeFinishMoveto(VdbeCursor*);
 SQLITE_PRIVATE int sqlite3VdbeCursorMoveto(VdbeCursor**, int*);
 SQLITE_PRIVATE int sqlite3VdbeCursorRestore(VdbeCursor*);
 SQLITE_PRIVATE u32 sqlite3VdbeSerialTypeLen(u32);
 SQLITE_PRIVATE u8 sqlite3VdbeOneByteSerialTypeLen(u8);
-SQLITE_PRIVATE u32 sqlite3VdbeSerialType(Mem*, int, u32*);
 SQLITE_PRIVATE u32 sqlite3VdbeSerialPut(unsigned char*, Mem*, u32);
 SQLITE_PRIVATE u32 sqlite3VdbeSerialGet(const unsigned char*, u32, Mem*);
 SQLITE_PRIVATE void sqlite3VdbeDeleteAuxData(sqlite3*, AuxData**, int, int);
@@ -20464,7 +21062,7 @@ SQLITE_PRIVATE int sqlite3VdbeBooleanValue(Mem*, int ifNull);
 SQLITE_PRIVATE void sqlite3VdbeIntegerAffinity(Mem*);
 SQLITE_PRIVATE int sqlite3VdbeMemRealify(Mem*);
 SQLITE_PRIVATE int sqlite3VdbeMemNumerify(Mem*);
-SQLITE_PRIVATE void sqlite3VdbeMemCast(Mem*,u8,u8);
+SQLITE_PRIVATE int sqlite3VdbeMemCast(Mem*,u8,u8);
 SQLITE_PRIVATE int sqlite3VdbeMemFromBtree(BtCursor*,u32,u32,Mem*);
 SQLITE_PRIVATE void sqlite3VdbeMemRelease(Mem *p);
 SQLITE_PRIVATE int sqlite3VdbeMemFinalize(Mem*, FuncDef*);
@@ -20530,7 +21128,7 @@ SQLITE_PRIVATE int sqlite3VdbeCheckFk(Vdbe *, int);
 
 #ifdef SQLITE_DEBUG
 SQLITE_PRIVATE   void sqlite3VdbePrintSql(Vdbe*);
-SQLITE_PRIVATE   void sqlite3VdbeMemPrettyPrint(Mem *pMem, char *zBuf);
+SQLITE_PRIVATE   void sqlite3VdbeMemPrettyPrint(Mem *pMem, StrAccum *pStr);
 #endif
 #ifndef SQLITE_OMIT_UTF16
 SQLITE_PRIVATE   int sqlite3VdbeMemTranslate(Mem*, u8);
@@ -20722,6 +21320,10 @@ static u32 countLookasideSlots(LookasideSlot *p){
 SQLITE_PRIVATE int sqlite3LookasideUsed(sqlite3 *db, int *pHighwater){
   u32 nInit = countLookasideSlots(db->lookaside.pInit);
   u32 nFree = countLookasideSlots(db->lookaside.pFree);
+#ifndef SQLITE_OMIT_TWOSIZE_LOOKASIDE
+  nInit += countLookasideSlots(db->lookaside.pSmallInit);
+  nFree += countLookasideSlots(db->lookaside.pSmallFree);
+#endif /* SQLITE_OMIT_TWOSIZE_LOOKASIDE */
   if( pHighwater ) *pHighwater = db->lookaside.nSlot - nInit;
   return db->lookaside.nSlot - (nInit+nFree);
 }
@@ -20754,6 +21356,15 @@ SQLITE_API int sqlite3_db_status(
           db->lookaside.pInit = db->lookaside.pFree;
           db->lookaside.pFree = 0;
         }
+#ifndef SQLITE_OMIT_TWOSIZE_LOOKASIDE
+        p = db->lookaside.pSmallFree;
+        if( p ){
+          while( p->pNext ) p = p->pNext;
+          p->pNext = db->lookaside.pSmallInit;
+          db->lookaside.pSmallInit = db->lookaside.pSmallFree;
+          db->lookaside.pSmallFree = 0;
+        }
+#endif
       }
       break;
     }
@@ -21305,7 +21916,7 @@ static int parseDateOrTime(
     return 0;
   }else if( sqlite3StrICmp(zDate,"now")==0 && sqlite3NotPureFunc(context) ){
     return setDateTimeToCurrent(context, p);
-  }else if( sqlite3AtoF(zDate, &r, sqlite3Strlen30(zDate), SQLITE_UTF8) ){
+  }else if( sqlite3AtoF(zDate, &r, sqlite3Strlen30(zDate), SQLITE_UTF8)>0 ){
     setRawDateNumber(p, r);
     return 0;
   }
@@ -21605,7 +22216,7 @@ static int parseModifier(
         r = p->s*1000.0 + 210866760000000.0;
         if( r>=0.0 && r<464269060800000.0 ){
           clearYMD_HMS_TZ(p);
-          p->iJD = (sqlite3_int64)r;
+          p->iJD = (sqlite3_int64)(r + 0.5);
           p->validJD = 1;
           p->rawS = 0;
           rc = 0;
@@ -21639,7 +22250,7 @@ static int parseModifier(
       ** date is already on the appropriate weekday, this is a no-op.
       */
       if( sqlite3_strnicmp(z, "weekday ", 8)==0
-               && sqlite3AtoF(&z[8], &r, sqlite3Strlen30(&z[8]), SQLITE_UTF8)
+               && sqlite3AtoF(&z[8], &r, sqlite3Strlen30(&z[8]), SQLITE_UTF8)>0
                && (n=(int)r)==r && n>=0 && r<7 ){
         sqlite3_int64 Z;
         computeYMD_HMS(p);
@@ -21698,7 +22309,7 @@ static int parseModifier(
       double rRounder;
       int i;
       for(n=1; z[n] && z[n]!=':' && !sqlite3Isspace(z[n]); n++){}
-      if( !sqlite3AtoF(z, &r, n, SQLITE_UTF8) ){
+      if( sqlite3AtoF(z, &r, n, SQLITE_UTF8)<=0 ){
         rc = 1;
         break;
       }
@@ -22384,7 +22995,7 @@ SQLITE_PRIVATE int sqlite3OsOpen(
   ** down into the VFS layer.  Some SQLITE_OPEN_ flags (for example,
   ** SQLITE_OPEN_FULLMUTEX or SQLITE_OPEN_SHAREDCACHE) are blocked before
   ** reaching the VFS. */
-  rc = pVfs->xOpen(pVfs, zPath, pFile, flags & 0x87f7f, pFlagsOut);
+  rc = pVfs->xOpen(pVfs, zPath, pFile, flags & 0x1087f7f, pFlagsOut);
   assert( rc==SQLITE_OK || pFile->pMethods==0 );
   return rc;
 }
@@ -22427,7 +23038,15 @@ SQLITE_PRIVATE void sqlite3OsDlClose(sqlite3_vfs *pVfs, void *pHandle){
 }
 #endif /* SQLITE_OMIT_LOAD_EXTENSION */
 SQLITE_PRIVATE int sqlite3OsRandomness(sqlite3_vfs *pVfs, int nByte, char *zBufOut){
-  return pVfs->xRandomness(pVfs, nByte, zBufOut);
+  if( sqlite3Config.iPrngSeed ){
+    memset(zBufOut, 0, nByte);
+    if( ALWAYS(nByte>(signed)sizeof(unsigned)) ) nByte = sizeof(unsigned int);
+    memcpy(zBufOut, &sqlite3Config.iPrngSeed, nByte);
+    return SQLITE_OK;
+  }else{
+    return pVfs->xRandomness(pVfs, nByte, zBufOut);
+  }
+  
 }
 SQLITE_PRIVATE int sqlite3OsSleep(sqlite3_vfs *pVfs, int nMicro){
   return pVfs->xSleep(pVfs, nMicro);
@@ -25894,7 +26513,7 @@ SQLITE_PRIVATE sqlite3_mutex_methods const *sqlite3DefaultMutex(void){
 ******************************************************************************
 **
 ** This file contains inline asm code for retrieving "high-performance"
-** counters for x86 class CPUs.
+** counters for x86 and x86_64 class CPUs.
 */
 #ifndef SQLITE_HWTIME_H
 #define SQLITE_HWTIME_H
@@ -25905,8 +26524,9 @@ SQLITE_PRIVATE sqlite3_mutex_methods const *sqlite3DefaultMutex(void){
 ** processor and returns that value.  This can be used for high-res
 ** profiling.
 */
-#if (defined(__GNUC__) || defined(_MSC_VER)) && \
-      (defined(i386) || defined(__i386__) || defined(_M_IX86))
+#if !defined(__STRICT_ANSI__) && \
+    (defined(__GNUC__) || defined(_MSC_VER)) && \
+    (defined(i386) || defined(__i386__) || defined(_M_IX86))
 
   #if defined(__GNUC__)
 
@@ -25927,7 +26547,7 @@ SQLITE_PRIVATE sqlite3_mutex_methods const *sqlite3DefaultMutex(void){
 
   #endif
 
-#elif (defined(__GNUC__) && defined(__x86_64__))
+#elif !defined(__STRICT_ANSI__) && (defined(__GNUC__) && defined(__x86_64__))
 
   __inline__ sqlite_uint64 sqlite3Hwtime(void){
       unsigned long val;
@@ -25935,7 +26555,7 @@ SQLITE_PRIVATE sqlite3_mutex_methods const *sqlite3DefaultMutex(void){
       return val;
   }
  
-#elif (defined(__GNUC__) && defined(__ppc__))
+#elif !defined(__STRICT_ANSI__) && (defined(__GNUC__) && defined(__ppc__))
 
   __inline__ sqlite_uint64 sqlite3Hwtime(void){
       unsigned long long retval;
@@ -25952,14 +26572,13 @@ SQLITE_PRIVATE sqlite3_mutex_methods const *sqlite3DefaultMutex(void){
 
 #else
 
-  #error Need implementation of sqlite3Hwtime() for your platform.
-
   /*
-  ** To compile without implementing sqlite3Hwtime() for your platform,
-  ** you can remove the above #error and use the following
-  ** stub function.  You will lose timing support for many
-  ** of the debugging and testing utilities, but it should at
-  ** least compile and run.
+  ** asm() is needed for hardware timing support.  Without asm(),
+  ** disable the sqlite3Hwtime() routine.
+  **
+  ** sqlite3Hwtime() is only used for some obscure debugging
+  ** and analysis configurations, not in any deliverable, so this
+  ** should not be a great loss.
   */
 SQLITE_PRIVATE   sqlite_uint64 sqlite3Hwtime(void){ return ((sqlite_uint64)0); }
 
@@ -26543,19 +27162,27 @@ SQLITE_API int sqlite3_release_memory(int n){
 #endif
 }
 
+/*
+** Default value of the hard heap limit.  0 means "no limit".
+*/
+#ifndef SQLITE_MAX_MEMORY
+# define SQLITE_MAX_MEMORY 0
+#endif
+
 /*
 ** State information local to the memory allocation subsystem.
 */
 static SQLITE_WSD struct Mem0Global {
   sqlite3_mutex *mutex;         /* Mutex to serialize access */
   sqlite3_int64 alarmThreshold; /* The soft heap limit */
+  sqlite3_int64 hardLimit;      /* The hard upper bound on memory */
 
   /*
   ** True if heap is nearly "full" where "full" is defined by the
   ** sqlite3_soft_heap_limit() setting.
   */
   int nearlyFull;
-} mem0 = { 0, 0, 0 };
+} mem0 = { 0, SQLITE_MAX_MEMORY, SQLITE_MAX_MEMORY, 0 };
 
 #define mem0 GLOBAL(struct Mem0Global, mem0)
 
@@ -26585,8 +27212,15 @@ SQLITE_API int sqlite3_memory_alarm(
 #endif
 
 /*
-** Set the soft heap-size limit for the library. Passing a zero or 
-** negative value indicates no limit.
+** Set the soft heap-size limit for the library.  An argument of
+** zero disables the limit.  A negative argument is a no-op used to
+** obtain the return value.
+**
+** The return value is the value of the heap limit just before this
+** interface was called.
+**
+** If the hard heap limit is enabled, then the soft heap limit cannot
+** be disabled nor raised above the hard heap limit.
 */
 SQLITE_API sqlite3_int64 sqlite3_soft_heap_limit64(sqlite3_int64 n){
   sqlite3_int64 priorLimit;
@@ -26602,6 +27236,9 @@ SQLITE_API sqlite3_int64 sqlite3_soft_heap_limit64(sqlite3_int64 n){
     sqlite3_mutex_leave(mem0.mutex);
     return priorLimit;
   }
+  if( mem0.hardLimit>0 && (n>mem0.hardLimit || n==0) ){
+    n = mem0.hardLimit;
+  }
   mem0.alarmThreshold = n;
   nUsed = sqlite3StatusValue(SQLITE_STATUS_MEMORY_USED);
   mem0.nearlyFull = (n>0 && n<=nUsed);
@@ -26615,6 +27252,37 @@ SQLITE_API void sqlite3_soft_heap_limit(int n){
   sqlite3_soft_heap_limit64(n);
 }
 
+/*
+** Set the hard heap-size limit for the library. An argument of zero
+** disables the hard heap limit.  A negative argument is a no-op used
+** to obtain the return value without affecting the hard heap limit.
+**
+** The return value is the value of the hard heap limit just prior to
+** calling this interface.
+**
+** Setting the hard heap limit will also activate the soft heap limit
+** and constrain the soft heap limit to be no more than the hard heap
+** limit.
+*/
+SQLITE_API sqlite3_int64 sqlite3_hard_heap_limit64(sqlite3_int64 n){
+  sqlite3_int64 priorLimit;
+#ifndef SQLITE_OMIT_AUTOINIT
+  int rc = sqlite3_initialize();
+  if( rc ) return -1;
+#endif
+  sqlite3_mutex_enter(mem0.mutex);
+  priorLimit = mem0.hardLimit;
+  if( n>=0 ){
+    mem0.hardLimit = n;
+    if( n<mem0.alarmThreshold || mem0.alarmThreshold==0 ){
+      mem0.alarmThreshold = n;
+    }
+  }
+  sqlite3_mutex_leave(mem0.mutex);
+  return priorLimit;
+}
+
+
 /*
 ** Initialize the memory allocation subsystem.
 */
@@ -26701,19 +27369,19 @@ static void mallocWithAlarm(int n, void **pp){
   ** following xRoundup() call. */
   nFull = sqlite3GlobalConfig.m.xRoundup(n);
 
-#ifdef SQLITE_MAX_MEMORY
-  if( sqlite3StatusValue(SQLITE_STATUS_MEMORY_USED)+nFull>SQLITE_MAX_MEMORY ){
-    *pp = 0;
-    return;
-  }
-#endif
-
   sqlite3StatusHighwater(SQLITE_STATUS_MALLOC_SIZE, n);
   if( mem0.alarmThreshold>0 ){
     sqlite3_int64 nUsed = sqlite3StatusValue(SQLITE_STATUS_MEMORY_USED);
     if( nUsed >= mem0.alarmThreshold - nFull ){
       mem0.nearlyFull = 1;
       sqlite3MallocAlarm(nFull);
+      if( mem0.hardLimit ){
+        nUsed = sqlite3StatusValue(SQLITE_STATUS_MEMORY_USED);
+        if( nUsed >= mem0.hardLimit - nFull ){
+          *pp = 0;
+          return;
+        }
+      }
     }else{
       mem0.nearlyFull = 0;
     }
@@ -26794,10 +27462,17 @@ SQLITE_PRIVATE int sqlite3MallocSize(void *p){
   assert( sqlite3MemdebugHasType(p, MEMTYPE_HEAP) );
   return sqlite3GlobalConfig.m.xSize(p);
 }
+static int lookasideMallocSize(sqlite3 *db, void *p){
+#ifndef SQLITE_OMIT_TWOSIZE_LOOKASIDE    
+  return p<db->lookaside.pMiddle ? db->lookaside.szTrue : LOOKASIDE_SMALL;
+#else
+  return db->lookaside.szTrue;
+#endif  
+}
 SQLITE_PRIVATE int sqlite3DbMallocSize(sqlite3 *db, void *p){
   assert( p!=0 );
-  if( db==0 || !isLookaside(db,p) ){
 #ifdef SQLITE_DEBUG
+  if( db==0 || !isLookaside(db,p) ){
     if( db==0 ){
       assert( sqlite3MemdebugNoType(p, (u8)~MEMTYPE_HEAP) );
       assert( sqlite3MemdebugHasType(p, MEMTYPE_HEAP) );
@@ -26805,12 +27480,23 @@ SQLITE_PRIVATE int sqlite3DbMallocSize(sqlite3 *db, void *p){
       assert( sqlite3MemdebugHasType(p, (MEMTYPE_LOOKASIDE|MEMTYPE_HEAP)) );
       assert( sqlite3MemdebugNoType(p, (u8)~(MEMTYPE_LOOKASIDE|MEMTYPE_HEAP)) );
     }
+  }
 #endif
-    return sqlite3GlobalConfig.m.xSize(p);
-  }else{
-    assert( sqlite3_mutex_held(db->mutex) );
-    return db->lookaside.sz;
+  if( db ){
+    if( ((uptr)p)<(uptr)(db->lookaside.pEnd) ){
+#ifndef SQLITE_OMIT_TWOSIZE_LOOKASIDE
+      if( ((uptr)p)>=(uptr)(db->lookaside.pMiddle) ){
+        assert( sqlite3_mutex_held(db->mutex) );
+        return LOOKASIDE_SMALL;
+      }
+#endif
+      if( ((uptr)p)>=(uptr)(db->lookaside.pStart) ){
+        assert( sqlite3_mutex_held(db->mutex) );
+        return db->lookaside.szTrue;
+      }
+    }
   }
+  return sqlite3GlobalConfig.m.xSize(p);
 }
 SQLITE_API sqlite3_uint64 sqlite3_msize(void *p){
   assert( sqlite3MemdebugNoType(p, (u8)~MEMTYPE_HEAP) );
@@ -26857,15 +27543,27 @@ SQLITE_PRIVATE void sqlite3DbFreeNN(sqlite3 *db, void *p){
       measureAllocationSize(db, p);
       return;
     }
-    if( isLookaside(db, p) ){
-      LookasideSlot *pBuf = (LookasideSlot*)p;
+    if( ((uptr)p)<(uptr)(db->lookaside.pEnd) ){
+#ifndef SQLITE_OMIT_TWOSIZE_LOOKASIDE
+      if( ((uptr)p)>=(uptr)(db->lookaside.pMiddle) ){
+        LookasideSlot *pBuf = (LookasideSlot*)p;
 #ifdef SQLITE_DEBUG
-      /* Trash all content in the buffer being freed */
-      memset(p, 0xaa, db->lookaside.sz);
+        memset(p, 0xaa, LOOKASIDE_SMALL);  /* Trash freed content */
 #endif
-      pBuf->pNext = db->lookaside.pFree;
-      db->lookaside.pFree = pBuf;
-      return;
+        pBuf->pNext = db->lookaside.pSmallFree;
+        db->lookaside.pSmallFree = pBuf;
+        return;
+      }
+#endif /* SQLITE_OMIT_TWOSIZE_LOOKASIDE */
+      if( ((uptr)p)>=(uptr)(db->lookaside.pStart) ){
+        LookasideSlot *pBuf = (LookasideSlot*)p;
+#ifdef SQLITE_DEBUG
+        memset(p, 0xaa, db->lookaside.szTrue);  /* Trash freed content */
+#endif
+        pBuf->pNext = db->lookaside.pFree;
+        db->lookaside.pFree = pBuf;
+        return;
+      }
     }
   }
   assert( sqlite3MemdebugHasType(p, (MEMTYPE_LOOKASIDE|MEMTYPE_HEAP)) );
@@ -27021,23 +27719,37 @@ SQLITE_PRIVATE void *sqlite3DbMallocRawNN(sqlite3 *db, u64 n){
   assert( db!=0 );
   assert( sqlite3_mutex_held(db->mutex) );
   assert( db->pnBytesFreed==0 );
-  if( db->lookaside.bDisable==0 ){
-    assert( db->mallocFailed==0 );
-    if( n>db->lookaside.sz ){
-      db->lookaside.anStat[1]++;
-    }else if( (pBuf = db->lookaside.pFree)!=0 ){
-      db->lookaside.pFree = pBuf->pNext;
+  if( n>db->lookaside.sz ){
+    if( !db->lookaside.bDisable ){
+      db->lookaside.anStat[1]++;      
+    }else if( db->mallocFailed ){
+      return 0;
+    }
+    return dbMallocRawFinish(db, n);
+  }
+#ifndef SQLITE_OMIT_TWOSIZE_LOOKASIDE
+  if( n<=LOOKASIDE_SMALL ){
+    if( (pBuf = db->lookaside.pSmallFree)!=0 ){
+      db->lookaside.pSmallFree = pBuf->pNext;
       db->lookaside.anStat[0]++;
       return (void*)pBuf;
-    }else if( (pBuf = db->lookaside.pInit)!=0 ){
-      db->lookaside.pInit = pBuf->pNext;
+    }else if( (pBuf = db->lookaside.pSmallInit)!=0 ){
+      db->lookaside.pSmallInit = pBuf->pNext;
       db->lookaside.anStat[0]++;
       return (void*)pBuf;
-    }else{
-      db->lookaside.anStat[2]++;
     }
-  }else if( db->mallocFailed ){
-    return 0;
+  }
+#endif
+  if( (pBuf = db->lookaside.pFree)!=0 ){
+    db->lookaside.pFree = pBuf->pNext;
+    db->lookaside.anStat[0]++;
+    return (void*)pBuf;
+  }else if( (pBuf = db->lookaside.pInit)!=0 ){
+    db->lookaside.pInit = pBuf->pNext;
+    db->lookaside.anStat[0]++;
+    return (void*)pBuf;
+  }else{
+    db->lookaside.anStat[2]++;
   }
 #else
   assert( db!=0 );
@@ -27061,7 +27773,16 @@ SQLITE_PRIVATE void *sqlite3DbRealloc(sqlite3 *db, void *p, u64 n){
   assert( db!=0 );
   if( p==0 ) return sqlite3DbMallocRawNN(db, n);
   assert( sqlite3_mutex_held(db->mutex) );
-  if( isLookaside(db,p) && n<=db->lookaside.sz ) return p;
+  if( ((uptr)p)<(uptr)db->lookaside.pEnd ){
+#ifndef SQLITE_OMIT_TWOSIZE_LOOKASIDE
+    if( ((uptr)p)>=(uptr)db->lookaside.pMiddle ){
+      if( n<=LOOKASIDE_SMALL ) return p;
+    }else
+#endif
+    if( ((uptr)p)>=(uptr)db->lookaside.pStart ){
+      if( n<=db->lookaside.szTrue ) return p;
+    }
+  }
   return dbReallocFinish(db, p, n);
 }
 static SQLITE_NOINLINE void *dbReallocFinish(sqlite3 *db, void *p, u64 n){
@@ -27072,7 +27793,7 @@ static SQLITE_NOINLINE void *dbReallocFinish(sqlite3 *db, void *p, u64 n){
     if( isLookaside(db, p) ){
       pNew = sqlite3DbMallocRawNN(db, n);
       if( pNew ){
-        memcpy(pNew, p, db->lookaside.sz);
+        memcpy(pNew, p, lookasideMallocSize(db, p));
         sqlite3DbFree(db, p);
       }
     }else{
@@ -27171,7 +27892,7 @@ SQLITE_PRIVATE void sqlite3OomFault(sqlite3 *db){
     if( db->nVdbeExec>0 ){
       db->u1.isInterrupted = 1;
     }
-    db->lookaside.bDisable++;
+    DisableLookaside;
     if( db->pParse ){
       db->pParse->rc = SQLITE_NOMEM_BKPT;
     }
@@ -27190,7 +27911,7 @@ SQLITE_PRIVATE void sqlite3OomClear(sqlite3 *db){
     db->mallocFailed = 0;
     db->u1.isInterrupted = 0;
     assert( db->lookaside.bDisable>0 );
-    db->lookaside.bDisable--;
+    EnableLookaside;
   }
 }
 
@@ -27331,6 +28052,12 @@ static const et_info fmtinfo[] = {
   {  'r', 10, 1, etORDINAL,    0,  0 },
 };
 
+/* Floating point constants used for rounding */
+static const double arRound[] = {
+  5.0e-01, 5.0e-02, 5.0e-03, 5.0e-04, 5.0e-05,
+  5.0e-06, 5.0e-07, 5.0e-08, 5.0e-09, 5.0e-10,
+};
+
 /*
 ** If SQLITE_OMIT_FLOATING_POINT is defined, then none of the floating point
 ** conversions will work.
@@ -27749,8 +28476,18 @@ SQLITE_API void sqlite3_str_vappendf(
         }
         if( xtype==etGENERIC && precision>0 ) precision--;
         testcase( precision>0xfff );
-        for(idx=precision&0xfff, rounder=0.5; idx>0; idx--, rounder*=0.1){}
-        if( xtype==etFLOAT ) realvalue += rounder;
+        idx = precision & 0xfff;
+        rounder = arRound[idx%10];
+        while( idx>=10 ){ rounder *= 1.0e-10; idx -= 10; }
+        if( xtype==etFLOAT ){
+          double rx = (double)realvalue;
+          sqlite3_uint64 u;
+          int ex;
+          memcpy(&u, &rx, sizeof(u));
+          ex = -1023 + (int)((u>>52)&0x7ff);
+          if( precision+(ex/3) < 15 ) rounder += realvalue*3e-16;
+          realvalue += rounder;
+        }
         /* Normalize realvalue to within 10.0 > realvalue >= 1.0 */
         exp = 0;
         if( sqlite3IsNaN((double)realvalue) ){
@@ -28582,7 +29319,7 @@ static void sqlite3TreeViewLine(TreeView *p, const char *zFormat, ...){
     va_start(ap, zFormat);
     sqlite3_str_vappendf(&acc, zFormat, ap);
     va_end(ap);
-    assert( acc.nChar>0 );
+    assert( acc.nChar>0 || acc.accError );
     sqlite3_str_append(&acc, "\n", 1);
   }
   sqlite3StrAccumFinish(&acc);
@@ -28622,7 +29359,7 @@ SQLITE_PRIVATE void sqlite3TreeViewWith(TreeView *pView, const With *pWith, u8 m
         char cSep = '(';
         int j;
         for(j=0; j<pCte->pCols->nExpr; j++){
-          sqlite3_str_appendf(&x, "%c%s", cSep, pCte->pCols->a[j].zName);
+          sqlite3_str_appendf(&x, "%c%s", cSep, pCte->pCols->a[j].zEName);
           cSep = ',';
         }
         sqlite3_str_appendf(&x, ")");
@@ -28647,7 +29384,7 @@ SQLITE_PRIVATE void sqlite3TreeViewSrcList(TreeView *pView, const SrcList *pSrc)
     StrAccum x;
     char zLine[100];
     sqlite3StrAccumInit(&x, 0, zLine, sizeof(zLine), 0);
-    sqlite3_str_appendf(&x, "{%d,*}", pItem->iCursor);
+    sqlite3_str_appendf(&x, "{%d:*}", pItem->iCursor);
     if( pItem->zDatabase ){
       sqlite3_str_appendf(&x, " %s.%s", pItem->zDatabase, pItem->zName);
     }else if( pItem->zName ){
@@ -28663,6 +29400,9 @@ SQLITE_PRIVATE void sqlite3TreeViewSrcList(TreeView *pView, const SrcList *pSrc)
     if( pItem->fg.jointype & JT_LEFT ){
       sqlite3_str_appendf(&x, " LEFT-JOIN");
     }
+    if( pItem->fg.fromDDL ){
+      sqlite3_str_appendf(&x, " DDL");
+    }
     sqlite3StrAccumFinish(&x);
     sqlite3TreeViewItem(pView, zLine, i<pSrc->nSrc-1); 
     if( pItem->pSelect ){
@@ -28692,13 +29432,17 @@ SQLITE_PRIVATE void sqlite3TreeViewSelect(TreeView *pView, const Select *p, u8 m
     sqlite3TreeViewPush(pView, 1);
   }
   do{
-    sqlite3TreeViewLine(pView,
-      "SELECT%s%s (%u/%p) selFlags=0x%x nSelectRow=%d",
-      ((p->selFlags & SF_Distinct) ? " DISTINCT" : ""),
-      ((p->selFlags & SF_Aggregate) ? " agg_flag" : ""),
-      p->selId, p, p->selFlags,
-      (int)p->nSelectRow
-    );
+    if( p->selFlags & SF_WhereBegin ){
+      sqlite3TreeViewLine(pView, "sqlite3WhereBegin()");
+    }else{
+      sqlite3TreeViewLine(pView,
+        "SELECT%s%s (%u/%p) selFlags=0x%x nSelectRow=%d",
+        ((p->selFlags & SF_Distinct) ? " DISTINCT" : ""),
+        ((p->selFlags & SF_Aggregate) ? " agg_flag" : ""),
+        p->selId, p, p->selFlags,
+        (int)p->nSelectRow
+      );
+    }
     if( cnt++ ) sqlite3TreeViewPop(pView);
     if( p->pPrior ){
       n = 1000;
@@ -28715,7 +29459,10 @@ SQLITE_PRIVATE void sqlite3TreeViewSelect(TreeView *pView, const Select *p, u8 m
       if( p->pWinDefn ) n++;
 #endif
     }
-    sqlite3TreeViewExprList(pView, p->pEList, (n--)>0, "result-set");
+    if( p->pEList ){
+      sqlite3TreeViewExprList(pView, p->pEList, n>0, "result-set");
+    }
+    n--;
 #ifndef SQLITE_OMIT_WINDOWFUNC
     if( p->pWin ){
       Window *pX;
@@ -28911,13 +29658,18 @@ SQLITE_PRIVATE void sqlite3TreeViewExpr(TreeView *pView, const Expr *pExpr, u8 m
     sqlite3TreeViewPop(pView);
     return;
   }
-  if( pExpr->flags ){
+  if( pExpr->flags || pExpr->affExpr ){
+    StrAccum x;
+    sqlite3StrAccumInit(&x, 0, zFlgs, sizeof(zFlgs), 0);
+    sqlite3_str_appendf(&x, " fg.af=%x.%c",
+      pExpr->flags, pExpr->affExpr ? pExpr->affExpr : 'n');
     if( ExprHasProperty(pExpr, EP_FromJoin) ){
-      sqlite3_snprintf(sizeof(zFlgs),zFlgs,"  flags=0x%x iRJT=%d",
-                       pExpr->flags, pExpr->iRightJoinTable);
-    }else{
-      sqlite3_snprintf(sizeof(zFlgs),zFlgs,"  flags=0x%x",pExpr->flags);
+      sqlite3_str_appendf(&x, " iRJT=%d", pExpr->iRightJoinTable);
+    }
+    if( ExprHasProperty(pExpr, EP_FromDDL) ){
+      sqlite3_str_appendf(&x, " DDL");
     }
+    sqlite3StrAccumFinish(&x);
   }else{
     zFlgs[0] = 0;
   }
@@ -28930,10 +29682,18 @@ SQLITE_PRIVATE void sqlite3TreeViewExpr(TreeView *pView, const Expr *pExpr, u8 m
     case TK_COLUMN: {
       if( pExpr->iTable<0 ){
         /* This only happens when coding check constraints */
-        sqlite3TreeViewLine(pView, "COLUMN(%d)%s", pExpr->iColumn, zFlgs);
+        char zOp2[16];
+        if( pExpr->op2 ){
+          sqlite3_snprintf(sizeof(zOp2),zOp2," op2=0x%02x",pExpr->op2);
+        }else{
+          zOp2[0] = 0;
+        }
+        sqlite3TreeViewLine(pView, "COLUMN(%d)%s%s",
+                                    pExpr->iColumn, zFlgs, zOp2);
       }else{
-        sqlite3TreeViewLine(pView, "{%d:%d}%s",
-                             pExpr->iTable, pExpr->iColumn, zFlgs);
+        sqlite3TreeViewLine(pView, "{%d:%d} pTab=%p%s",
+                        pExpr->iTable, pExpr->iColumn,
+                        pExpr->y.pTab, zFlgs);
       }
       if( ExprHasProperty(pExpr, EP_FixedCol) ){
         sqlite3TreeViewExpr(pView, pExpr->pLeft, 0);
@@ -29030,7 +29790,7 @@ SQLITE_PRIVATE void sqlite3TreeViewExpr(TreeView *pView, const Expr *pExpr, u8 m
       };
       assert( pExpr->op2==TK_IS || pExpr->op2==TK_ISNOT );
       assert( pExpr->pRight );
-      assert( pExpr->pRight->op==TK_TRUEFALSE );
+      assert( sqlite3ExprSkipCollate(pExpr->pRight)->op==TK_TRUEFALSE );
       x = (pExpr->op2==TK_ISNOT)*2 + sqlite3ExprTruthValue(pExpr->pRight);
       zUniOp = azOp[x];
       break;
@@ -29043,7 +29803,14 @@ SQLITE_PRIVATE void sqlite3TreeViewExpr(TreeView *pView, const Expr *pExpr, u8 m
     }
 
     case TK_COLLATE: {
-      sqlite3TreeViewLine(pView, "COLLATE %Q", pExpr->u.zToken);
+      /* COLLATE operators without the EP_Collate flag are intended to
+      ** emulate collation associated with a table column.  These show
+      ** up in the treeview output as "SOFT-COLLATE".  Explicit COLLATE
+      ** operators that appear in the original SQL always have the
+      ** EP_Collate bit set and appear in treeview output as just "COLLATE" */
+      sqlite3TreeViewLine(pView, "%sCOLLATE %Q%s",
+        !ExprHasProperty(pExpr, EP_Collate) ? "SOFT-" : "",
+        pExpr->u.zToken, zFlgs);
       sqlite3TreeViewExpr(pView, pExpr->pLeft, 0);
       break;
     }
@@ -29058,16 +29825,27 @@ SQLITE_PRIVATE void sqlite3TreeViewExpr(TreeView *pView, const Expr *pExpr, u8 m
       }else{
         pFarg = pExpr->x.pList;
 #ifndef SQLITE_OMIT_WINDOWFUNC
-        pWin = pExpr->y.pWin;
+        pWin = ExprHasProperty(pExpr, EP_WinFunc) ? pExpr->y.pWin : 0;
 #else
         pWin = 0;
 #endif 
       }
       if( pExpr->op==TK_AGG_FUNCTION ){
-        sqlite3TreeViewLine(pView, "AGG_FUNCTION%d %Q",
-                             pExpr->op2, pExpr->u.zToken);
+        sqlite3TreeViewLine(pView, "AGG_FUNCTION%d %Q%s",
+                             pExpr->op2, pExpr->u.zToken, zFlgs);
+      }else if( pExpr->op2!=0 ){
+        const char *zOp2;
+        char zBuf[8];
+        sqlite3_snprintf(sizeof(zBuf),zBuf,"0x%02x",pExpr->op2);
+        zOp2 = zBuf;
+        if( pExpr->op2==NC_IsCheck ) zOp2 = "NC_IsCheck";
+        if( pExpr->op2==NC_IdxExpr ) zOp2 = "NC_IdxExpr";
+        if( pExpr->op2==NC_PartIdx ) zOp2 = "NC_PartIdx";
+        if( pExpr->op2==NC_GenCol ) zOp2 = "NC_GenCol";
+        sqlite3TreeViewLine(pView, "FUNCTION %Q%s op2=%s",
+                            pExpr->u.zToken, zFlgs, zOp2);
       }else{
-        sqlite3TreeViewLine(pView, "FUNCTION %Q", pExpr->u.zToken);
+        sqlite3TreeViewLine(pView, "FUNCTION %Q%s", pExpr->u.zToken, zFlgs);
       }
       if( pFarg ){
         sqlite3TreeViewExprList(pView, pFarg, pWin!=0, 0);
@@ -29144,7 +29922,7 @@ SQLITE_PRIVATE void sqlite3TreeViewExpr(TreeView *pView, const Expr *pExpr, u8 m
 #ifndef SQLITE_OMIT_TRIGGER
     case TK_RAISE: {
       const char *zType = "unk";
-      switch( pExpr->affinity ){
+      switch( pExpr->affExpr ){
         case OE_Rollback:   zType = "rollback";  break;
         case OE_Abort:      zType = "abort";     break;
         case OE_Fail:       zType = "fail";      break;
@@ -29161,7 +29939,9 @@ SQLITE_PRIVATE void sqlite3TreeViewExpr(TreeView *pView, const Expr *pExpr, u8 m
       break;
     }
     case TK_VECTOR: {
-      sqlite3TreeViewBareExprList(pView, pExpr->x.pList, "VECTOR");
+      char *z = sqlite3_mprintf("VECTOR%s",zFlgs);
+      sqlite3TreeViewBareExprList(pView, pExpr->x.pList, z);
+      sqlite3_free(z);
       break;
     }
     case TK_SELECT_COLUMN: {
@@ -29185,7 +29965,7 @@ SQLITE_PRIVATE void sqlite3TreeViewExpr(TreeView *pView, const Expr *pExpr, u8 m
     sqlite3TreeViewExpr(pView, pExpr->pRight, 0);
   }else if( zUniOp ){
     sqlite3TreeViewLine(pView, "%s%s", zUniOp, zFlgs);
-    sqlite3TreeViewExpr(pView, pExpr->pLeft, 0);
+   sqlite3TreeViewExpr(pView, pExpr->pLeft, 0);
   }
   sqlite3TreeViewPop(pView);
 }
@@ -29207,8 +29987,9 @@ SQLITE_PRIVATE void sqlite3TreeViewBareExprList(
     sqlite3TreeViewLine(pView, "%s", zLabel);
     for(i=0; i<pList->nExpr; i++){
       int j = pList->a[i].u.x.iOrderByCol;
-      char *zName = pList->a[i].zName;
+      char *zName = pList->a[i].zEName;
       int moreToFollow = i<pList->nExpr - 1;
+      if( pList->a[i].eEName!=ENAME_NAME ) zName = 0;
       if( j || zName ){
         sqlite3TreeViewPush(pView, moreToFollow);
         moreToFollow = 0;
@@ -29875,9 +30656,11 @@ SQLITE_PRIVATE SQLITE_NOINLINE int sqlite3VdbeMemTranslate(Mem *pMem, u8 desired
 
 #if defined(TRANSLATE_TRACE) && defined(SQLITE_DEBUG)
   {
-    char zBuf[100];
-    sqlite3VdbeMemPrettyPrint(pMem, zBuf);
-    fprintf(stderr, "INPUT:  %s\n", zBuf);
+    StrAccum acc;
+    char zBuf[1000];
+    sqlite3StrAccumInit(&acc, 0, zBuf, sizeof(zBuf), 0);  
+    sqlite3VdbeMemPrettyPrint(pMem, &acc);
+    fprintf(stderr, "INPUT:  %s\n", sqlite3StrAccumFinish(&acc));
   }
 #endif
 
@@ -29985,9 +30768,11 @@ SQLITE_PRIVATE SQLITE_NOINLINE int sqlite3VdbeMemTranslate(Mem *pMem, u8 desired
 translate_out:
 #if defined(TRANSLATE_TRACE) && defined(SQLITE_DEBUG)
   {
-    char zBuf[100];
-    sqlite3VdbeMemPrettyPrint(pMem, zBuf);
-    fprintf(stderr, "OUTPUT: %s\n", zBuf);
+    StrAccum acc;
+    char zBuf[1000];
+    sqlite3StrAccumInit(&acc, 0, zBuf, sizeof(zBuf), 0);  
+    sqlite3VdbeMemPrettyPrint(pMem, &acc);
+    fprintf(stderr, "OUTPUT: %s\n", sqlite3StrAccumFinish(&acc));
   }
 #endif
   return SQLITE_OK;
@@ -30212,8 +30997,8 @@ SQLITE_PRIVATE void sqlite3UtfSelfTest(void){
 */
 /* #include "sqliteInt.h" */
 /* #include <stdarg.h> */
-#if HAVE_ISNAN || SQLITE_HAVE_ISNAN
-# include <math.h>
+#ifndef SQLITE_OMIT_FLOATING_POINT
+#include <math.h>
 #endif
 
 /*
@@ -30255,47 +31040,11 @@ SQLITE_PRIVATE int sqlite3FaultSim(int iTest){
 #ifndef SQLITE_OMIT_FLOATING_POINT
 /*
 ** Return true if the floating point value is Not a Number (NaN).
-**
-** Use the math library isnan() function if compiled with SQLITE_HAVE_ISNAN.
-** Otherwise, we have our own implementation that works on most systems.
 */
 SQLITE_PRIVATE int sqlite3IsNaN(double x){
-  int rc;   /* The value return */
-#if !SQLITE_HAVE_ISNAN && !HAVE_ISNAN
-  /*
-  ** Systems that support the isnan() library function should probably
-  ** make use of it by compiling with -DSQLITE_HAVE_ISNAN.  But we have
-  ** found that many systems do not have a working isnan() function so
-  ** this implementation is provided as an alternative.
-  **
-  ** This NaN test sometimes fails if compiled on GCC with -ffast-math.
-  ** On the other hand, the use of -ffast-math comes with the following
-  ** warning:
-  **
-  **      This option [-ffast-math] should never be turned on by any
-  **      -O option since it can result in incorrect output for programs
-  **      which depend on an exact implementation of IEEE or ISO 
-  **      rules/specifications for math functions.
-  **
-  ** Under MSVC, this NaN test may fail if compiled with a floating-
-  ** point precision mode other than /fp:precise.  From the MSDN 
-  ** documentation:
-  **
-  **      The compiler [with /fp:precise] will properly handle comparisons 
-  **      involving NaN. For example, x != x evaluates to true if x is NaN 
-  **      ...
-  */
-#ifdef __FAST_MATH__
-# error SQLite will not work correctly with the -ffast-math option of GCC.
-#endif
-  volatile double y = x;
-  volatile double z = y;
-  rc = (y!=z);
-#else  /* if HAVE_ISNAN */
-  rc = isnan(x);
-#endif /* HAVE_ISNAN */
-  testcase( rc );
-  return rc;
+  u64 y;
+  memcpy(&y,&x,sizeof(y));
+  return IsNaN(y);
 }
 #endif /* SQLITE_OMIT_FLOATING_POINT */
 
@@ -30425,6 +31174,7 @@ SQLITE_PRIVATE void sqlite3ErrorMsg(Parse *pParse, const char *zFormat, ...){
     sqlite3DbFree(db, pParse->zErrMsg);
     pParse->zErrMsg = zMsg;
     pParse->rc = SQLITE_ERROR;
+    pParse->pWith = 0;
   }
 }
 
@@ -30517,12 +31267,18 @@ SQLITE_API int sqlite3_stricmp(const char *zLeft, const char *zRight){
 }
 SQLITE_PRIVATE int sqlite3StrICmp(const char *zLeft, const char *zRight){
   unsigned char *a, *b;
-  int c;
+  int c, x;
   a = (unsigned char *)zLeft;
   b = (unsigned char *)zRight;
   for(;;){
-    c = (int)UpperToLower[*a] - (int)UpperToLower[*b];
-    if( c || *a==0 ) break;
+    c = *a;
+    x = *b;
+    if( c==x ){
+      if( c==0 ) break;
+    }else{
+      c = (int)UpperToLower[c] - (int)UpperToLower[x];
+      if( c ) break;
+    }
     a++;
     b++;
   }
@@ -30550,15 +31306,15 @@ SQLITE_API int sqlite3_strnicmp(const char *zLeft, const char *zRight, int N){
 static LONGDOUBLE_TYPE sqlite3Pow10(int E){
 #if defined(_MSC_VER)
   static const LONGDOUBLE_TYPE x[] = {
-    1.0e+001,
-    1.0e+002,
-    1.0e+004,
-    1.0e+008,
-    1.0e+016,
-    1.0e+032,
-    1.0e+064,
-    1.0e+128,
-    1.0e+256
+    1.0e+001L,
+    1.0e+002L,
+    1.0e+004L,
+    1.0e+008L,
+    1.0e+016L,
+    1.0e+032L,
+    1.0e+064L,
+    1.0e+128L,
+    1.0e+256L
   };
   LONGDOUBLE_TYPE r = 1.0;
   int i;
@@ -30588,8 +31344,15 @@ static LONGDOUBLE_TYPE sqlite3Pow10(int E){
 ** uses the encoding enc.  The string is not necessarily zero-terminated.
 **
 ** Return TRUE if the result is a valid real number (or integer) and FALSE
-** if the string is empty or contains extraneous text.  Valid numbers
-** are in one of these formats:
+** if the string is empty or contains extraneous text.  More specifically
+** return
+**      1          =>  The input string is a pure integer
+**      2 or more  =>  The input has a decimal point or eNNN clause
+**      0 or less  =>  The input string is not a valid number
+**     -1          =>  Not a valid number, but has a valid prefix which 
+**                     includes a decimal point and/or an eNNN clause
+**
+** Valid numbers are in one of these formats:
 **
 **    [+-]digits[E[+-]digits]
 **    [+-]digits.[digits][E[+-]digits]
@@ -30602,10 +31365,13 @@ static LONGDOUBLE_TYPE sqlite3Pow10(int E){
 ** returns FALSE but it still converts the prefix and writes the result
 ** into *pResult.
 */
+#if defined(_MSC_VER)
+#pragma warning(disable : 4756)
+#endif
 SQLITE_PRIVATE int sqlite3AtoF(const char *z, double *pResult, int length, u8 enc){
 #ifndef SQLITE_OMIT_FLOATING_POINT
   int incr;
-  const char *zEnd = z + length;
+  const char *zEnd;
   /* sign * significand * (10 ^ (esign * exponent)) */
   int sign = 1;    /* sign of significand */
   i64 s = 0;       /* significand */
@@ -30614,20 +31380,25 @@ SQLITE_PRIVATE int sqlite3AtoF(const char *z, double *pResult, int length, u8 en
   int e = 0;       /* exponent */
   int eValid = 1;  /* True exponent is either not used or is well-formed */
   double result;
-  int nDigits = 0;
-  int nonNum = 0;  /* True if input contains UTF16 with high byte non-zero */
+  int nDigit = 0;  /* Number of digits processed */
+  int eType = 1;   /* 1: pure integer,  2+: fractional  -1 or less: bad UTF16 */
 
   assert( enc==SQLITE_UTF8 || enc==SQLITE_UTF16LE || enc==SQLITE_UTF16BE );
   *pResult = 0.0;   /* Default return value, in case of an error */
+  if( length==0 ) return 0;
 
   if( enc==SQLITE_UTF8 ){
     incr = 1;
+    zEnd = z + length;
   }else{
     int i;
     incr = 2;
+    length &= ~1;
     assert( SQLITE_UTF16LE==2 && SQLITE_UTF16BE==3 );
+    testcase( enc==SQLITE_UTF16LE );
+    testcase( enc==SQLITE_UTF16BE );
     for(i=3-enc; i<length && z[i]==0; i+=2){}
-    nonNum = i<length;
+    if( i<length ) eType = -100;
     zEnd = &z[i^1];
     z += (enc&1);
   }
@@ -30645,27 +31416,30 @@ SQLITE_PRIVATE int sqlite3AtoF(const char *z, double *pResult, int length, u8 en
   }
 
   /* copy max significant digits to significand */
-  while( z<zEnd && sqlite3Isdigit(*z) && s<((LARGEST_INT64-9)/10) ){
+  while( z<zEnd && sqlite3Isdigit(*z) ){
     s = s*10 + (*z - '0');
-    z+=incr; nDigits++;
+    z+=incr; nDigit++;
+    if( s>=((LARGEST_INT64-9)/10) ){
+      /* skip non-significant significand digits
+      ** (increase exponent by d to shift decimal left) */
+      while( z<zEnd && sqlite3Isdigit(*z) ){ z+=incr; d++; }
+    }
   }
-
-  /* skip non-significant significand digits
-  ** (increase exponent by d to shift decimal left) */
-  while( z<zEnd && sqlite3Isdigit(*z) ){ z+=incr; nDigits++; d++; }
   if( z>=zEnd ) goto do_atof_calc;
 
   /* if decimal point is present */
   if( *z=='.' ){
     z+=incr;
+    eType++;
     /* copy digits from after decimal to significand
     ** (decrease exponent by d to shift decimal right) */
     while( z<zEnd && sqlite3Isdigit(*z) ){
       if( s<((LARGEST_INT64-9)/10) ){
         s = s*10 + (*z - '0');
         d--;
+        nDigit++;
       }
-      z+=incr; nDigits++;
+      z+=incr;
     }
   }
   if( z>=zEnd ) goto do_atof_calc;
@@ -30674,6 +31448,7 @@ SQLITE_PRIVATE int sqlite3AtoF(const char *z, double *pResult, int length, u8 en
   if( *z=='e' || *z=='E' ){
     z+=incr;
     eValid = 0;
+    eType++;
 
     /* This branch is needed to avoid a (harmless) buffer overread.  The 
     ** special comment alerts the mutation tester that the correct answer
@@ -30772,11 +31547,20 @@ do_atof_calc:
   *pResult = result;
 
   /* return true if number and no extra non-whitespace chracters after */
-  return z==zEnd && nDigits>0 && eValid && nonNum==0;
+  if( z==zEnd && nDigit>0 && eValid && eType>0 ){
+    return eType;
+  }else if( eType>=2 && (eType==3 || eValid) && nDigit>0 ){
+    return -1;
+  }else{
+    return 0;
+  }
 #else
   return !sqlite3Atoi64(z, pResult, length, enc);
 #endif /* SQLITE_OMIT_FLOATING_POINT */
 }
+#if defined(_MSC_VER)
+#pragma warning(default : 4756)
+#endif
 
 /*
 ** Compare the 19-character string zNum against the text representation
@@ -30815,6 +31599,7 @@ static int compare2pow63(const char *zNum, int incr){
 **
 ** Returns:
 **
+**    -1    Not even a prefix of the input text looks like an integer
 **     0    Successful transformation.  Fits in a 64-bit signed integer.
 **     1    Excess non-space text after the integer value
 **     2    Integer too large for a 64-bit signed integer or is malformed
@@ -30874,9 +31659,9 @@ SQLITE_PRIVATE int sqlite3Atoi64(const char *zNum, i64 *pNum, int length, u8 enc
     *pNum = (i64)u;
   }
   rc = 0;
-  if( (i==0 && zStart==zNum)     /* No digits */
-   || nonNum                     /* UTF16 with high-order bytes non-zero */
-  ){
+  if( i==0 && zStart==zNum ){    /* No digits */
+    rc = -1;
+  }else if( nonNum ){            /* UTF16 with high-order bytes non-zero */
     rc = 1;
   }else if( &zNum[i]<zEnd ){     /* Extra bytes at the end */
     int jj = i;
@@ -31107,23 +31892,12 @@ SQLITE_PRIVATE int sqlite3PutVarint(unsigned char *p, u64 v){
 SQLITE_PRIVATE u8 sqlite3GetVarint(const unsigned char *p, u64 *v){
   u32 a,b,s;
 
-  a = *p;
-  /* a: p0 (unmasked) */
-  if (!(a&0x80))
-  {
-    *v = a;
+  if( ((signed char*)p)[0]>=0 ){
+    *v = *p;
     return 1;
   }
-
-  p++;
-  b = *p;
-  /* b: p1 (unmasked) */
-  if (!(b&0x80))
-  {
-    a &= 0x7f;
-    a = a<<7;
-    a |= b;
-    *v = a;
+  if( ((signed char*)p)[1]>=0 ){
+    *v = ((u32)(p[0]&0x7f)<<7) | p[1];
     return 2;
   }
 
@@ -31131,8 +31905,9 @@ SQLITE_PRIVATE u8 sqlite3GetVarint(const unsigned char *p, u64 *v){
   assert( SLOT_2_0 == ((0x7f<<14) | (0x7f)) );
   assert( SLOT_4_2_0 == ((0xfU<<28) | (0x7f<<14) | (0x7f)) );
 
-  p++;
-  a = a<<14;
+  a = ((u32)p[0])<<14;
+  b = p[1];
+  p += 2;
   a |= *p;
   /* a: p0<<14 | p2 (unmasked) */
   if (!(a&0x80))
@@ -31715,7 +32490,7 @@ SQLITE_PRIVATE LogEst sqlite3LogEstFromDouble(double x){
 #endif /* SQLITE_OMIT_VIRTUALTABLE */
 
 #if defined(SQLITE_ENABLE_STMT_SCANSTATUS) || \
-    defined(SQLITE_ENABLE_STAT3_OR_STAT4) || \
+    defined(SQLITE_ENABLE_STAT4) || \
     defined(SQLITE_EXPLAIN_ESTIMATED_ROWS)
 /*
 ** Convert a LogEst into an integer.
@@ -31733,7 +32508,7 @@ SQLITE_PRIVATE u64 sqlite3LogEstToInt(LogEst x){
     defined(SQLITE_EXPLAIN_ESTIMATED_ROWS)
   if( x>60 ) return (u64)LARGEST_INT64;
 #else
-  /* If only SQLITE_ENABLE_STAT3_OR_STAT4 is on, then the largest input
+  /* If only SQLITE_ENABLE_STAT4 is on, then the largest input
   ** possible to this routine is 310, resulting in a maximum x of 31 */
   assert( x<=60 );
 #endif
@@ -32156,30 +32931,30 @@ SQLITE_PRIVATE const char *sqlite3OpcodeName(int i){
     /*  23 */ "SeekLE"           OpHelp("key=r[P3@P4]"),
     /*  24 */ "SeekGE"           OpHelp("key=r[P3@P4]"),
     /*  25 */ "SeekGT"           OpHelp("key=r[P3@P4]"),
-    /*  26 */ "IfNoHope"         OpHelp("key=r[P3@P4]"),
-    /*  27 */ "NoConflict"       OpHelp("key=r[P3@P4]"),
-    /*  28 */ "NotFound"         OpHelp("key=r[P3@P4]"),
-    /*  29 */ "Found"            OpHelp("key=r[P3@P4]"),
-    /*  30 */ "SeekRowid"        OpHelp("intkey=r[P3]"),
-    /*  31 */ "NotExists"        OpHelp("intkey=r[P3]"),
-    /*  32 */ "Last"             OpHelp(""),
-    /*  33 */ "IfSmaller"        OpHelp(""),
-    /*  34 */ "SorterSort"       OpHelp(""),
-    /*  35 */ "Sort"             OpHelp(""),
-    /*  36 */ "Rewind"           OpHelp(""),
-    /*  37 */ "IdxLE"            OpHelp("key=r[P3@P4]"),
-    /*  38 */ "IdxGT"            OpHelp("key=r[P3@P4]"),
-    /*  39 */ "IdxLT"            OpHelp("key=r[P3@P4]"),
-    /*  40 */ "IdxGE"            OpHelp("key=r[P3@P4]"),
-    /*  41 */ "RowSetRead"       OpHelp("r[P3]=rowset(P1)"),
-    /*  42 */ "RowSetTest"       OpHelp("if r[P3] in rowset(P1) goto P2"),
+    /*  26 */ "IfNotOpen"        OpHelp("if( !csr[P1] ) goto P2"),
+    /*  27 */ "IfNoHope"         OpHelp("key=r[P3@P4]"),
+    /*  28 */ "NoConflict"       OpHelp("key=r[P3@P4]"),
+    /*  29 */ "NotFound"         OpHelp("key=r[P3@P4]"),
+    /*  30 */ "Found"            OpHelp("key=r[P3@P4]"),
+    /*  31 */ "SeekRowid"        OpHelp("intkey=r[P3]"),
+    /*  32 */ "NotExists"        OpHelp("intkey=r[P3]"),
+    /*  33 */ "Last"             OpHelp(""),
+    /*  34 */ "IfSmaller"        OpHelp(""),
+    /*  35 */ "SorterSort"       OpHelp(""),
+    /*  36 */ "Sort"             OpHelp(""),
+    /*  37 */ "Rewind"           OpHelp(""),
+    /*  38 */ "IdxLE"            OpHelp("key=r[P3@P4]"),
+    /*  39 */ "IdxGT"            OpHelp("key=r[P3@P4]"),
+    /*  40 */ "IdxLT"            OpHelp("key=r[P3@P4]"),
+    /*  41 */ "IdxGE"            OpHelp("key=r[P3@P4]"),
+    /*  42 */ "RowSetRead"       OpHelp("r[P3]=rowset(P1)"),
     /*  43 */ "Or"               OpHelp("r[P3]=(r[P1] || r[P2])"),
     /*  44 */ "And"              OpHelp("r[P3]=(r[P1] && r[P2])"),
-    /*  45 */ "Program"          OpHelp(""),
-    /*  46 */ "FkIfZero"         OpHelp("if fkctr[P1]==0 goto P2"),
-    /*  47 */ "IfPos"            OpHelp("if r[P1]>0 then r[P1]-=P3, goto P2"),
-    /*  48 */ "IfNotZero"        OpHelp("if r[P1]!=0 then r[P1]--, goto P2"),
-    /*  49 */ "DecrJumpZero"     OpHelp("if (--r[P1])==0 goto P2"),
+    /*  45 */ "RowSetTest"       OpHelp("if r[P3] in rowset(P1) goto P2"),
+    /*  46 */ "Program"          OpHelp(""),
+    /*  47 */ "FkIfZero"         OpHelp("if fkctr[P1]==0 goto P2"),
+    /*  48 */ "IfPos"            OpHelp("if r[P1]>0 then r[P1]-=P3, goto P2"),
+    /*  49 */ "IfNotZero"        OpHelp("if r[P1]!=0 then r[P1]--, goto P2"),
     /*  50 */ "IsNull"           OpHelp("if r[P1]==NULL goto P2"),
     /*  51 */ "NotNull"          OpHelp("if r[P1]!=NULL goto P2"),
     /*  52 */ "Ne"               OpHelp("IF r[P3]!=r[P1]"),
@@ -32189,83 +32964,83 @@ SQLITE_PRIVATE const char *sqlite3OpcodeName(int i){
     /*  56 */ "Lt"               OpHelp("IF r[P3]<r[P1]"),
     /*  57 */ "Ge"               OpHelp("IF r[P3]>=r[P1]"),
     /*  58 */ "ElseNotEq"        OpHelp(""),
-    /*  59 */ "IncrVacuum"       OpHelp(""),
-    /*  60 */ "VNext"            OpHelp(""),
-    /*  61 */ "Init"             OpHelp("Start at P2"),
-    /*  62 */ "PureFunc0"        OpHelp(""),
-    /*  63 */ "Function0"        OpHelp("r[P3]=func(r[P2@P5])"),
-    /*  64 */ "PureFunc"         OpHelp(""),
-    /*  65 */ "Function"         OpHelp("r[P3]=func(r[P2@P5])"),
-    /*  66 */ "Return"           OpHelp(""),
-    /*  67 */ "EndCoroutine"     OpHelp(""),
-    /*  68 */ "HaltIfNull"       OpHelp("if r[P3]=null halt"),
-    /*  69 */ "Halt"             OpHelp(""),
-    /*  70 */ "Integer"          OpHelp("r[P2]=P1"),
-    /*  71 */ "Int64"            OpHelp("r[P2]=P4"),
-    /*  72 */ "String"           OpHelp("r[P2]='P4' (len=P1)"),
-    /*  73 */ "Null"             OpHelp("r[P2..P3]=NULL"),
-    /*  74 */ "SoftNull"         OpHelp("r[P1]=NULL"),
-    /*  75 */ "Blob"             OpHelp("r[P2]=P4 (len=P1)"),
-    /*  76 */ "Variable"         OpHelp("r[P2]=parameter(P1,P4)"),
-    /*  77 */ "Move"             OpHelp("r[P2@P3]=r[P1@P3]"),
-    /*  78 */ "Copy"             OpHelp("r[P2@P3+1]=r[P1@P3+1]"),
-    /*  79 */ "SCopy"            OpHelp("r[P2]=r[P1]"),
-    /*  80 */ "IntCopy"          OpHelp("r[P2]=r[P1]"),
-    /*  81 */ "ResultRow"        OpHelp("output=r[P1@P2]"),
-    /*  82 */ "CollSeq"          OpHelp(""),
-    /*  83 */ "AddImm"           OpHelp("r[P1]=r[P1]+P2"),
-    /*  84 */ "RealAffinity"     OpHelp(""),
-    /*  85 */ "Cast"             OpHelp("affinity(r[P1])"),
-    /*  86 */ "Permutation"      OpHelp(""),
-    /*  87 */ "Compare"          OpHelp("r[P1@P3] <-> r[P2@P3]"),
-    /*  88 */ "IsTrue"           OpHelp("r[P2] = coalesce(r[P1]==TRUE,P3) ^ P4"),
-    /*  89 */ "Offset"           OpHelp("r[P3] = sqlite_offset(P1)"),
-    /*  90 */ "Column"           OpHelp("r[P3]=PX"),
-    /*  91 */ "Affinity"         OpHelp("affinity(r[P1@P2])"),
-    /*  92 */ "MakeRecord"       OpHelp("r[P3]=mkrec(r[P1@P2])"),
-    /*  93 */ "Count"            OpHelp("r[P2]=count()"),
-    /*  94 */ "ReadCookie"       OpHelp(""),
-    /*  95 */ "SetCookie"        OpHelp(""),
-    /*  96 */ "BitAnd"           OpHelp("r[P3]=r[P1]&r[P2]"),
-    /*  97 */ "BitOr"            OpHelp("r[P3]=r[P1]|r[P2]"),
-    /*  98 */ "ShiftLeft"        OpHelp("r[P3]=r[P2]<<r[P1]"),
-    /*  99 */ "ShiftRight"       OpHelp("r[P3]=r[P2]>>r[P1]"),
-    /* 100 */ "Add"              OpHelp("r[P3]=r[P1]+r[P2]"),
-    /* 101 */ "Subtract"         OpHelp("r[P3]=r[P2]-r[P1]"),
-    /* 102 */ "Multiply"         OpHelp("r[P3]=r[P1]*r[P2]"),
-    /* 103 */ "Divide"           OpHelp("r[P3]=r[P2]/r[P1]"),
-    /* 104 */ "Remainder"        OpHelp("r[P3]=r[P2]%r[P1]"),
-    /* 105 */ "Concat"           OpHelp("r[P3]=r[P2]+r[P1]"),
-    /* 106 */ "ReopenIdx"        OpHelp("root=P2 iDb=P3"),
-    /* 107 */ "BitNot"           OpHelp("r[P2]= ~r[P1]"),
-    /* 108 */ "OpenRead"         OpHelp("root=P2 iDb=P3"),
-    /* 109 */ "OpenWrite"        OpHelp("root=P2 iDb=P3"),
-    /* 110 */ "String8"          OpHelp("r[P2]='P4'"),
-    /* 111 */ "OpenDup"          OpHelp(""),
-    /* 112 */ "OpenAutoindex"    OpHelp("nColumn=P2"),
-    /* 113 */ "OpenEphemeral"    OpHelp("nColumn=P2"),
-    /* 114 */ "SorterOpen"       OpHelp(""),
-    /* 115 */ "SequenceTest"     OpHelp("if( cursor[P1].ctr++ ) pc = P2"),
-    /* 116 */ "OpenPseudo"       OpHelp("P3 columns in r[P2]"),
-    /* 117 */ "Close"            OpHelp(""),
-    /* 118 */ "ColumnsUsed"      OpHelp(""),
-    /* 119 */ "SeekHit"          OpHelp("seekHit=P2"),
-    /* 120 */ "Sequence"         OpHelp("r[P2]=cursor[P1].ctr++"),
-    /* 121 */ "NewRowid"         OpHelp("r[P2]=rowid"),
-    /* 122 */ "Insert"           OpHelp("intkey=r[P3] data=r[P2]"),
-    /* 123 */ "Delete"           OpHelp(""),
-    /* 124 */ "ResetCount"       OpHelp(""),
-    /* 125 */ "SorterCompare"    OpHelp("if key(P1)!=trim(r[P3],P4) goto P2"),
-    /* 126 */ "SorterData"       OpHelp("r[P2]=data"),
-    /* 127 */ "RowData"          OpHelp("r[P2]=data"),
-    /* 128 */ "Rowid"            OpHelp("r[P2]=rowid"),
-    /* 129 */ "NullRow"          OpHelp(""),
-    /* 130 */ "SeekEnd"          OpHelp(""),
-    /* 131 */ "SorterInsert"     OpHelp("key=r[P2]"),
-    /* 132 */ "IdxInsert"        OpHelp("key=r[P2]"),
-    /* 133 */ "IdxDelete"        OpHelp("key=r[P2@P3]"),
-    /* 134 */ "DeferredSeek"     OpHelp("Move P3 to P1.rowid if needed"),
-    /* 135 */ "IdxRowid"         OpHelp("r[P2]=rowid"),
+    /*  59 */ "DecrJumpZero"     OpHelp("if (--r[P1])==0 goto P2"),
+    /*  60 */ "IncrVacuum"       OpHelp(""),
+    /*  61 */ "VNext"            OpHelp(""),
+    /*  62 */ "Init"             OpHelp("Start at P2"),
+    /*  63 */ "PureFunc"         OpHelp("r[P3]=func(r[P2@P5])"),
+    /*  64 */ "Function"         OpHelp("r[P3]=func(r[P2@P5])"),
+    /*  65 */ "Return"           OpHelp(""),
+    /*  66 */ "EndCoroutine"     OpHelp(""),
+    /*  67 */ "HaltIfNull"       OpHelp("if r[P3]=null halt"),
+    /*  68 */ "Halt"             OpHelp(""),
+    /*  69 */ "Integer"          OpHelp("r[P2]=P1"),
+    /*  70 */ "Int64"            OpHelp("r[P2]=P4"),
+    /*  71 */ "String"           OpHelp("r[P2]='P4' (len=P1)"),
+    /*  72 */ "Null"             OpHelp("r[P2..P3]=NULL"),
+    /*  73 */ "SoftNull"         OpHelp("r[P1]=NULL"),
+    /*  74 */ "Blob"             OpHelp("r[P2]=P4 (len=P1)"),
+    /*  75 */ "Variable"         OpHelp("r[P2]=parameter(P1,P4)"),
+    /*  76 */ "Move"             OpHelp("r[P2@P3]=r[P1@P3]"),
+    /*  77 */ "Copy"             OpHelp("r[P2@P3+1]=r[P1@P3+1]"),
+    /*  78 */ "SCopy"            OpHelp("r[P2]=r[P1]"),
+    /*  79 */ "IntCopy"          OpHelp("r[P2]=r[P1]"),
+    /*  80 */ "ResultRow"        OpHelp("output=r[P1@P2]"),
+    /*  81 */ "CollSeq"          OpHelp(""),
+    /*  82 */ "AddImm"           OpHelp("r[P1]=r[P1]+P2"),
+    /*  83 */ "RealAffinity"     OpHelp(""),
+    /*  84 */ "Cast"             OpHelp("affinity(r[P1])"),
+    /*  85 */ "Permutation"      OpHelp(""),
+    /*  86 */ "Compare"          OpHelp("r[P1@P3] <-> r[P2@P3]"),
+    /*  87 */ "IsTrue"           OpHelp("r[P2] = coalesce(r[P1]==TRUE,P3) ^ P4"),
+    /*  88 */ "Offset"           OpHelp("r[P3] = sqlite_offset(P1)"),
+    /*  89 */ "Column"           OpHelp("r[P3]=PX"),
+    /*  90 */ "Affinity"         OpHelp("affinity(r[P1@P2])"),
+    /*  91 */ "MakeRecord"       OpHelp("r[P3]=mkrec(r[P1@P2])"),
+    /*  92 */ "Count"            OpHelp("r[P2]=count()"),
+    /*  93 */ "ReadCookie"       OpHelp(""),
+    /*  94 */ "SetCookie"        OpHelp(""),
+    /*  95 */ "ReopenIdx"        OpHelp("root=P2 iDb=P3"),
+    /*  96 */ "OpenRead"         OpHelp("root=P2 iDb=P3"),
+    /*  97 */ "OpenWrite"        OpHelp("root=P2 iDb=P3"),
+    /*  98 */ "OpenDup"          OpHelp(""),
+    /*  99 */ "OpenAutoindex"    OpHelp("nColumn=P2"),
+    /* 100 */ "OpenEphemeral"    OpHelp("nColumn=P2"),
+    /* 101 */ "BitAnd"           OpHelp("r[P3]=r[P1]&r[P2]"),
+    /* 102 */ "BitOr"            OpHelp("r[P3]=r[P1]|r[P2]"),
+    /* 103 */ "ShiftLeft"        OpHelp("r[P3]=r[P2]<<r[P1]"),
+    /* 104 */ "ShiftRight"       OpHelp("r[P3]=r[P2]>>r[P1]"),
+    /* 105 */ "Add"              OpHelp("r[P3]=r[P1]+r[P2]"),
+    /* 106 */ "Subtract"         OpHelp("r[P3]=r[P2]-r[P1]"),
+    /* 107 */ "Multiply"         OpHelp("r[P3]=r[P1]*r[P2]"),
+    /* 108 */ "Divide"           OpHelp("r[P3]=r[P2]/r[P1]"),
+    /* 109 */ "Remainder"        OpHelp("r[P3]=r[P2]%r[P1]"),
+    /* 110 */ "Concat"           OpHelp("r[P3]=r[P2]+r[P1]"),
+    /* 111 */ "SorterOpen"       OpHelp(""),
+    /* 112 */ "BitNot"           OpHelp("r[P2]= ~r[P1]"),
+    /* 113 */ "SequenceTest"     OpHelp("if( cursor[P1].ctr++ ) pc = P2"),
+    /* 114 */ "OpenPseudo"       OpHelp("P3 columns in r[P2]"),
+    /* 115 */ "String8"          OpHelp("r[P2]='P4'"),
+    /* 116 */ "Close"            OpHelp(""),
+    /* 117 */ "ColumnsUsed"      OpHelp(""),
+    /* 118 */ "SeekHit"          OpHelp("seekHit=P2"),
+    /* 119 */ "Sequence"         OpHelp("r[P2]=cursor[P1].ctr++"),
+    /* 120 */ "NewRowid"         OpHelp("r[P2]=rowid"),
+    /* 121 */ "Insert"           OpHelp("intkey=r[P3] data=r[P2]"),
+    /* 122 */ "Delete"           OpHelp(""),
+    /* 123 */ "ResetCount"       OpHelp(""),
+    /* 124 */ "SorterCompare"    OpHelp("if key(P1)!=trim(r[P3],P4) goto P2"),
+    /* 125 */ "SorterData"       OpHelp("r[P2]=data"),
+    /* 126 */ "RowData"          OpHelp("r[P2]=data"),
+    /* 127 */ "Rowid"            OpHelp("r[P2]=rowid"),
+    /* 128 */ "NullRow"          OpHelp(""),
+    /* 129 */ "SeekEnd"          OpHelp(""),
+    /* 130 */ "SorterInsert"     OpHelp("key=r[P2]"),
+    /* 131 */ "IdxInsert"        OpHelp("key=r[P2]"),
+    /* 132 */ "IdxDelete"        OpHelp("key=r[P2@P3]"),
+    /* 133 */ "DeferredSeek"     OpHelp("Move P3 to P1.rowid if needed"),
+    /* 134 */ "IdxRowid"         OpHelp("r[P2]=rowid"),
+    /* 135 */ "FinishSeek"       OpHelp(""),
     /* 136 */ "Destroy"          OpHelp(""),
     /* 137 */ "Clear"            OpHelp(""),
     /* 138 */ "ResetSorter"      OpHelp(""),
@@ -32275,12 +33050,12 @@ SQLITE_PRIVATE const char *sqlite3OpcodeName(int i){
     /* 142 */ "LoadAnalysis"     OpHelp(""),
     /* 143 */ "DropTable"        OpHelp(""),
     /* 144 */ "DropIndex"        OpHelp(""),
-    /* 145 */ "Real"             OpHelp("r[P2]=P4"),
-    /* 146 */ "DropTrigger"      OpHelp(""),
-    /* 147 */ "IntegrityCk"      OpHelp(""),
-    /* 148 */ "RowSetAdd"        OpHelp("rowset(P1)=r[P2]"),
-    /* 149 */ "Param"            OpHelp(""),
-    /* 150 */ "FkCounter"        OpHelp("fkctr[P1]+=P2"),
+    /* 145 */ "DropTrigger"      OpHelp(""),
+    /* 146 */ "IntegrityCk"      OpHelp(""),
+    /* 147 */ "RowSetAdd"        OpHelp("rowset(P1)=r[P2]"),
+    /* 148 */ "Param"            OpHelp(""),
+    /* 149 */ "FkCounter"        OpHelp("fkctr[P1]+=P2"),
+    /* 150 */ "Real"             OpHelp("r[P2]=P4"),
     /* 151 */ "MemMax"           OpHelp("r[P1]=max(r[P1],r[P2])"),
     /* 152 */ "OffsetLimit"      OpHelp("if r[P1]>0 then r[P2]=r[P1]+max(0,r[P3]) else r[P2]=(-1)"),
     /* 153 */ "AggInverse"       OpHelp("accum=r[P3] inverse(r[P2@P5])"),
@@ -32289,20 +33064,23 @@ SQLITE_PRIVATE const char *sqlite3OpcodeName(int i){
     /* 156 */ "AggValue"         OpHelp("r[P3]=value N=P2"),
     /* 157 */ "AggFinal"         OpHelp("accum=r[P1] N=P2"),
     /* 158 */ "Expire"           OpHelp(""),
-    /* 159 */ "TableLock"        OpHelp("iDb=P1 root=P2 write=P3"),
-    /* 160 */ "VBegin"           OpHelp(""),
-    /* 161 */ "VCreate"          OpHelp(""),
-    /* 162 */ "VDestroy"         OpHelp(""),
-    /* 163 */ "VOpen"            OpHelp(""),
-    /* 164 */ "VColumn"          OpHelp("r[P3]=vcolumn(P2)"),
-    /* 165 */ "VRename"          OpHelp(""),
-    /* 166 */ "Pagecount"        OpHelp(""),
-    /* 167 */ "MaxPgcnt"         OpHelp(""),
-    /* 168 */ "Trace"            OpHelp(""),
-    /* 169 */ "CursorHint"       OpHelp(""),
-    /* 170 */ "Noop"             OpHelp(""),
-    /* 171 */ "Explain"          OpHelp(""),
-    /* 172 */ "Abortable"        OpHelp(""),
+    /* 159 */ "CursorLock"       OpHelp(""),
+    /* 160 */ "CursorUnlock"     OpHelp(""),
+    /* 161 */ "TableLock"        OpHelp("iDb=P1 root=P2 write=P3"),
+    /* 162 */ "VBegin"           OpHelp(""),
+    /* 163 */ "VCreate"          OpHelp(""),
+    /* 164 */ "VDestroy"         OpHelp(""),
+    /* 165 */ "VOpen"            OpHelp(""),
+    /* 166 */ "VColumn"          OpHelp("r[P3]=vcolumn(P2)"),
+    /* 167 */ "VRename"          OpHelp(""),
+    /* 168 */ "Pagecount"        OpHelp(""),
+    /* 169 */ "MaxPgcnt"         OpHelp(""),
+    /* 170 */ "Trace"            OpHelp(""),
+    /* 171 */ "CursorHint"       OpHelp(""),
+    /* 172 */ "ReleaseReg"       OpHelp("release r[P1@P2] mask P3"),
+    /* 173 */ "Noop"             OpHelp(""),
+    /* 174 */ "Explain"          OpHelp(""),
+    /* 175 */ "Abortable"        OpHelp(""),
   };
   return azName[i];
 }
@@ -32417,13 +33195,29 @@ SQLITE_PRIVATE const char *sqlite3OpcodeName(int i){
 # include <sys/param.h>
 #endif /* SQLITE_ENABLE_LOCKING_STYLE */
 
-#if defined(__APPLE__) && ((__MAC_OS_X_VERSION_MIN_REQUIRED > 1050) || \
-                           (__IPHONE_OS_VERSION_MIN_REQUIRED > 2000))
-#  if (!defined(TARGET_OS_EMBEDDED) || (TARGET_OS_EMBEDDED==0)) \
-       && (!defined(TARGET_IPHONE_SIMULATOR) || (TARGET_IPHONE_SIMULATOR==0))
-#    define HAVE_GETHOSTUUID 1
-#  else
-#    warning "gethostuuid() is disabled."
+/*
+** Try to determine if gethostuuid() is available based on standard
+** macros.  This might sometimes compute the wrong value for some
+** obscure platforms.  For those cases, simply compile with one of
+** the following:
+**
+**    -DHAVE_GETHOSTUUID=0
+**    -DHAVE_GETHOSTUUID=1
+**
+** None if this matters except when building on Apple products with
+** -DSQLITE_ENABLE_LOCKING_STYLE.
+*/
+#ifndef HAVE_GETHOSTUUID
+# define HAVE_GETHOSTUUID 0
+# if defined(__APPLE__) && ((__MAC_OS_X_VERSION_MIN_REQUIRED > 1050) || \
+                            (__IPHONE_OS_VERSION_MIN_REQUIRED > 2000))
+#    if (!defined(TARGET_OS_EMBEDDED) || (TARGET_OS_EMBEDDED==0)) \
+         && (!defined(TARGET_IPHONE_SIMULATOR) || (TARGET_IPHONE_SIMULATOR==0))
+#      undef HAVE_GETHOSTUUID
+#      define HAVE_GETHOSTUUID 1
+#    else
+#      warning "gethostuuid() is disabled."
+#    endif
 #  endif
 #endif
 
@@ -32650,7 +33444,7 @@ static pid_t randomnessPid = 0;
 ******************************************************************************
 **
 ** This file contains inline asm code for retrieving "high-performance"
-** counters for x86 class CPUs.
+** counters for x86 and x86_64 class CPUs.
 */
 #ifndef SQLITE_HWTIME_H
 #define SQLITE_HWTIME_H
@@ -32661,8 +33455,9 @@ static pid_t randomnessPid = 0;
 ** processor and returns that value.  This can be used for high-res
 ** profiling.
 */
-#if (defined(__GNUC__) || defined(_MSC_VER)) && \
-      (defined(i386) || defined(__i386__) || defined(_M_IX86))
+#if !defined(__STRICT_ANSI__) && \
+    (defined(__GNUC__) || defined(_MSC_VER)) && \
+    (defined(i386) || defined(__i386__) || defined(_M_IX86))
 
   #if defined(__GNUC__)
 
@@ -32683,7 +33478,7 @@ static pid_t randomnessPid = 0;
 
   #endif
 
-#elif (defined(__GNUC__) && defined(__x86_64__))
+#elif !defined(__STRICT_ANSI__) && (defined(__GNUC__) && defined(__x86_64__))
 
   __inline__ sqlite_uint64 sqlite3Hwtime(void){
       unsigned long val;
@@ -32691,7 +33486,7 @@ static pid_t randomnessPid = 0;
       return val;
   }
  
-#elif (defined(__GNUC__) && defined(__ppc__))
+#elif !defined(__STRICT_ANSI__) && (defined(__GNUC__) && defined(__ppc__))
 
   __inline__ sqlite_uint64 sqlite3Hwtime(void){
       unsigned long long retval;
@@ -32708,14 +33503,13 @@ static pid_t randomnessPid = 0;
 
 #else
 
-  #error Need implementation of sqlite3Hwtime() for your platform.
-
   /*
-  ** To compile without implementing sqlite3Hwtime() for your platform,
-  ** you can remove the above #error and use the following
-  ** stub function.  You will lose timing support for many
-  ** of the debugging and testing utilities, but it should at
-  ** least compile and run.
+  ** asm() is needed for hardware timing support.  Without asm(),
+  ** disable the sqlite3Hwtime() routine.
+  **
+  ** sqlite3Hwtime() is only used for some obscure debugging
+  ** and analysis configurations, not in any deliverable, so this
+  ** should not be a great loss.
   */
 SQLITE_PRIVATE   sqlite_uint64 sqlite3Hwtime(void){ return ((sqlite_uint64)0); }
 
@@ -33031,13 +33825,14 @@ static struct unix_syscall {
 #if defined(__linux__) && defined(SQLITE_ENABLE_BATCH_ATOMIC_WRITE)
 # ifdef __ANDROID__
   { "ioctl", (sqlite3_syscall_ptr)(int(*)(int, int, ...))ioctl, 0 },
+#define osIoctl ((int(*)(int,int,...))aSyscall[28].pCurrent)
 # else
   { "ioctl",         (sqlite3_syscall_ptr)ioctl,          0 },
+#define osIoctl ((int(*)(int,unsigned long,...))aSyscall[28].pCurrent)
 # endif
 #else
   { "ioctl",         (sqlite3_syscall_ptr)0,              0 },
 #endif
-#define osIoctl ((int(*)(int,int,...))aSyscall[28].pCurrent)
 
 }; /* End of the overrideable system calls */
 
@@ -36178,7 +36973,7 @@ static int openDirectory(const char *zFilename, int *pFd){
     if( zDirname[0]!='/' ) zDirname[0] = '.';
     zDirname[1] = 0;
   }
-  fd = robust_open(zDirname, O_RDONLY|O_BINARY, 0);
+  fd = robust_open(zDirname, O_RDONLY|O_BINARY|O_NOFOLLOW, 0);
   if( fd>=0 ){
     OSTRACE(("OPENDIR %-3d %s\n", fd, zDirname));
   }
@@ -37069,10 +37864,12 @@ static int unixOpenSharedMemory(unixFile *pDbFd){
 
     if( pInode->bProcessLock==0 ){
       if( 0==sqlite3_uri_boolean(pDbFd->zPath, "readonly_shm", 0) ){
-        pShmNode->hShm = robust_open(zShm, O_RDWR|O_CREAT,(sStat.st_mode&0777));
+        pShmNode->hShm = robust_open(zShm, O_RDWR|O_CREAT|O_NOFOLLOW,
+                                     (sStat.st_mode&0777));
       }
       if( pShmNode->hShm<0 ){
-        pShmNode->hShm = robust_open(zShm, O_RDONLY, (sStat.st_mode&0777));
+        pShmNode->hShm = robust_open(zShm, O_RDONLY|O_NOFOLLOW,
+                                     (sStat.st_mode&0777));
         if( pShmNode->hShm<0 ){
           rc = unixLogError(SQLITE_CANTOPEN_BKPT, "open", zShm);
           goto shm_open_err;
@@ -38279,6 +39076,7 @@ static UnixUnusedFd *findReusableFd(const char *zPath, int flags){
       UnixUnusedFd **pp;
       assert( sqlite3_mutex_notheld(pInode->pLockMutex) );
       sqlite3_mutex_enter(pInode->pLockMutex);
+      flags &= (SQLITE_OPEN_READONLY|SQLITE_OPEN_READWRITE);
       for(pp=&pInode->pUnused; *pp && (*pp)->flags!=flags; pp=&((*pp)->pNext));
       pUnused = *pp;
       if( pUnused ){
@@ -38332,7 +39130,7 @@ static int getFileMode(
 ** If the SQLITE_ENABLE_8_3_NAMES option is enabled, then the
 ** original filename is unavailable.  But 8_3_NAMES is only used for
 ** FAT filesystems and permissions do not matter there, so just use
-** the default permissions.
+** the default permissions.  In 8_3_NAMES mode, leave *pMode set to zero.
 */
 static int findCreateFileMode(
   const char *zPath,              /* Path of file (possibly) being created */
@@ -38421,7 +39219,7 @@ static int unixOpen(
   unixFile *p = (unixFile *)pFile;
   int fd = -1;                   /* File descriptor returned by open() */
   int openFlags = 0;             /* Flags to pass to open() */
-  int eType = flags&0xFFFFFF00;  /* Type of file to open */
+  int eType = flags&0x0FFF00;  /* Type of file to open */
   int noLock;                    /* True to omit locking primitives */
   int rc = SQLITE_OK;            /* Function Return Code */
   int ctrlFlags = 0;             /* UNIXFILE_* flags */
@@ -38531,7 +39329,7 @@ static int unixOpen(
   if( isReadWrite ) openFlags |= O_RDWR;
   if( isCreate )    openFlags |= O_CREAT;
   if( isExclusive ) openFlags |= (O_EXCL|O_NOFOLLOW);
-  openFlags |= (O_LARGEFILE|O_BINARY);
+  openFlags |= (O_LARGEFILE|O_BINARY|O_NOFOLLOW);
 
   if( fd<0 ){
     mode_t openMode;              /* Permissions to create file with */
@@ -38567,11 +39365,19 @@ static int unixOpen(
       goto open_finished;
     }
 
-    /* If this process is running as root and if creating a new rollback
-    ** journal or WAL file, set the ownership of the journal or WAL to be
-    ** the same as the original database.
+    /* The owner of the rollback journal or WAL file should always be the
+    ** same as the owner of the database file.  Try to ensure that this is
+    ** the case.  The chown() system call will be a no-op if the current
+    ** process lacks root privileges, be we should at least try.  Without
+    ** this step, if a root process opens a database file, it can leave
+    ** behinds a journal/WAL that is owned by root and hence make the
+    ** database inaccessible to unprivileged processes.
+    **
+    ** If openMode==0, then that means uid and gid are not set correctly
+    ** (probably because SQLite is configured to use 8+3 filename mode) and
+    ** in that case we do not want to attempt the chown().
     */
-    if( flags & (SQLITE_OPEN_WAL|SQLITE_OPEN_MAIN_JOURNAL) ){
+    if( openMode && (flags & (SQLITE_OPEN_WAL|SQLITE_OPEN_MAIN_JOURNAL))!=0 ){
       robustFchown(fd, uid, gid);
     }
   }
@@ -38582,7 +39388,8 @@ static int unixOpen(
 
   if( p->pPreallocatedUnused ){
     p->pPreallocatedUnused->fd = fd;
-    p->pPreallocatedUnused->flags = flags;
+    p->pPreallocatedUnused->flags = 
+                          flags & (SQLITE_OPEN_READONLY|SQLITE_OPEN_READWRITE);
   }
 
   if( isDelete ){
@@ -38740,7 +39547,8 @@ static int unixAccess(
 
   if( flags==SQLITE_ACCESS_EXISTS ){
     struct stat buf;
-    *pResOut = (0==osStat(zPath, &buf) && buf.st_size>0);
+    *pResOut = 0==osStat(zPath, &buf) &&
+                (!S_ISREG(buf.st_mode) || buf.st_size>0);
   }else{
     *pResOut = osAccess(zPath, W_OK|R_OK)==0;
   }
@@ -38794,7 +39602,7 @@ static int unixFullPathname(
 #else
   int rc = SQLITE_OK;
   int nByte;
-  int nLink = 1;                /* Number of symbolic links followed so far */
+  int nLink = 0;                /* Number of symbolic links followed so far */
   const char *zIn = zPath;      /* Input path for each iteration of loop */
   char *zDel = 0;
 
@@ -38823,10 +39631,11 @@ static int unixFullPathname(
     }
 
     if( bLink ){
+      nLink++;
       if( zDel==0 ){
         zDel = sqlite3_malloc(nOut);
         if( zDel==0 ) rc = SQLITE_NOMEM_BKPT;
-      }else if( ++nLink>SQLITE_MAX_SYMLINKS ){
+      }else if( nLink>=SQLITE_MAX_SYMLINKS ){
         rc = SQLITE_CANTOPEN_BKPT;
       }
 
@@ -38862,6 +39671,7 @@ static int unixFullPathname(
   }while( rc==SQLITE_OK );
 
   sqlite3_free(zDel);
+  if( rc==SQLITE_OK && nLink ) rc = SQLITE_OK_SYMLINK;
   return rc;
 #endif   /* HAVE_READLINK && HAVE_LSTAT */
 }
@@ -39347,7 +40157,7 @@ static int proxyCreateUnixFile(
   int fd = -1;
   unixFile *pNew;
   int rc = SQLITE_OK;
-  int openFlags = O_RDWR | O_CREAT;
+  int openFlags = O_RDWR | O_CREAT | O_NOFOLLOW;
   sqlite3_vfs dummyVfs;
   int terrno = 0;
   UnixUnusedFd *pUnused = NULL;
@@ -39377,7 +40187,7 @@ static int proxyCreateUnixFile(
     }
   }
   if( fd<0 ){
-    openFlags = O_RDONLY;
+    openFlags = O_RDONLY | O_NOFOLLOW;
     fd = robust_open(path, openFlags, 0);
     terrno = errno;
   }
@@ -39428,7 +40238,7 @@ SQLITE_API int sqlite3_hostid_num = 0;
 
 #define PROXY_HOSTIDLEN    16  /* conch file host id length */
 
-#ifdef HAVE_GETHOSTUUID
+#if HAVE_GETHOSTUUID
 /* Not always defined in the headers as it ought to be */
 extern int gethostuuid(uuid_t id, const struct timespec *wait);
 #endif
@@ -39439,7 +40249,7 @@ extern int gethostuuid(uuid_t id, const struct timespec *wait);
 static int proxyGetHostID(unsigned char *pHostID, int *pError){
   assert(PROXY_HOSTIDLEN == sizeof(uuid_t));
   memset(pHostID, 0, PROXY_HOSTIDLEN);
-#ifdef HAVE_GETHOSTUUID
+#if HAVE_GETHOSTUUID
   {
     struct timespec timeout = {1, 0}; /* 1 sec timeout */
     if( gethostuuid(pHostID, &timeout) ){
@@ -39503,7 +40313,7 @@ static int proxyBreakConchLock(unixFile *pFile, uuid_t myHostID){
     goto end_breaklock;
   }
   /* write it out to the temporary break file */
-  fd = robust_open(tPath, (O_RDWR|O_CREAT|O_EXCL), 0);
+  fd = robust_open(tPath, (O_RDWR|O_CREAT|O_EXCL|O_NOFOLLOW), 0);
   if( fd<0 ){
     sqlite3_snprintf(sizeof(errmsg), errmsg, "create failed (%d)", errno);
     goto end_breaklock;
@@ -40113,7 +40923,7 @@ static int proxyFileControl(sqlite3_file *id, int op, void *pArg){
       assert( 0 );  /* The call assures that only valid opcodes are sent */
     }
   }
-  /*NOTREACHED*/
+  /*NOTREACHED*/ assert(0);
   return SQLITE_ERROR;
 }
 
@@ -40461,7 +41271,7 @@ SQLITE_API int sqlite3_os_end(void){
 ******************************************************************************
 **
 ** This file contains inline asm code for retrieving "high-performance"
-** counters for x86 class CPUs.
+** counters for x86 and x86_64 class CPUs.
 */
 #ifndef SQLITE_HWTIME_H
 #define SQLITE_HWTIME_H
@@ -40472,8 +41282,9 @@ SQLITE_API int sqlite3_os_end(void){
 ** processor and returns that value.  This can be used for high-res
 ** profiling.
 */
-#if (defined(__GNUC__) || defined(_MSC_VER)) && \
-      (defined(i386) || defined(__i386__) || defined(_M_IX86))
+#if !defined(__STRICT_ANSI__) && \
+    (defined(__GNUC__) || defined(_MSC_VER)) && \
+    (defined(i386) || defined(__i386__) || defined(_M_IX86))
 
   #if defined(__GNUC__)
 
@@ -40494,7 +41305,7 @@ SQLITE_API int sqlite3_os_end(void){
 
   #endif
 
-#elif (defined(__GNUC__) && defined(__x86_64__))
+#elif !defined(__STRICT_ANSI__) && (defined(__GNUC__) && defined(__x86_64__))
 
   __inline__ sqlite_uint64 sqlite3Hwtime(void){
       unsigned long val;
@@ -40502,7 +41313,7 @@ SQLITE_API int sqlite3_os_end(void){
       return val;
   }
  
-#elif (defined(__GNUC__) && defined(__ppc__))
+#elif !defined(__STRICT_ANSI__) && (defined(__GNUC__) && defined(__ppc__))
 
   __inline__ sqlite_uint64 sqlite3Hwtime(void){
       unsigned long long retval;
@@ -40519,14 +41330,13 @@ SQLITE_API int sqlite3_os_end(void){
 
 #else
 
-  #error Need implementation of sqlite3Hwtime() for your platform.
-
   /*
-  ** To compile without implementing sqlite3Hwtime() for your platform,
-  ** you can remove the above #error and use the following
-  ** stub function.  You will lose timing support for many
-  ** of the debugging and testing utilities, but it should at
-  ** least compile and run.
+  ** asm() is needed for hardware timing support.  Without asm(),
+  ** disable the sqlite3Hwtime() routine.
+  **
+  ** sqlite3Hwtime() is only used for some obscure debugging
+  ** and analysis configurations, not in any deliverable, so this
+  ** should not be a great loss.
   */
 SQLITE_PRIVATE   sqlite_uint64 sqlite3Hwtime(void){ return ((sqlite_uint64)0); }
 
@@ -44798,6 +45608,7 @@ static int winShmMap(
     rc = winOpenSharedMemory(pDbFd);
     if( rc!=SQLITE_OK ) return rc;
     pShm = pDbFd->pShm;
+    assert( pShm!=0 );
   }
   pShmNode = pShm->pShmNode;
 
@@ -45100,6 +45911,7 @@ static int winFetch(sqlite3_file *fd, i64 iOff, int nAmt, void **pp){
       }
     }
     if( pFd->mmapSize >= iOff+nAmt ){
+      assert( pFd->pMapRegion!=0 );
       *pp = &((u8 *)pFd->pMapRegion)[iOff];
       pFd->nFetchOut++;
     }
@@ -48002,9 +48814,10 @@ static int numberOfCachePages(PCache *p){
     ** suggested cache size is set to N. */
     return p->szCache;
   }else{
-    /* IMPLEMENTATION-OF: R-61436-13639 If the argument N is negative, then
-    ** the number of cache pages is adjusted to use approximately abs(N*1024)
-    ** bytes of memory. */
+    /* IMPLEMANTATION-OF: R-59858-46238 If the argument N is negative, then the
+    ** number of cache pages is adjusted to be a number of pages that would
+    ** use approximately abs(N*1024) bytes of memory based on the current
+    ** page size. */
     return (int)((-1024*(i64)p->szCache)/(p->szPage+p->szExtra));
   }
 }
@@ -48020,6 +48833,7 @@ SQLITE_PRIVATE int sqlite3PcacheInitialize(void){
     ** built-in default page cache is used instead of the application defined
     ** page cache. */
     sqlite3PCacheSetDefault();
+    assert( sqlite3GlobalConfig.pcache2.xInit!=0 );
   }
   return sqlite3GlobalConfig.pcache2.xInit(sqlite3GlobalConfig.pcache2.pArg);
 }
@@ -49066,6 +49880,7 @@ static PgHdr1 *pcache1AllocPage(PCache1 *pCache, int benignMalloc){
 
   assert( sqlite3_mutex_held(pCache->pGroup->mutex) );
   if( pCache->pFree || (pCache->nPage==0 && pcache1InitBulk(pCache)) ){
+    assert( pCache->pFree!=0 );
     p = pCache->pFree;
     pCache->pFree = p->pNext;
     p->pNext = 0;
@@ -49089,13 +49904,15 @@ static PgHdr1 *pcache1AllocPage(PCache1 *pCache, int benignMalloc){
     }
 #else
     pPg = pcache1Alloc(pCache->szAlloc);
-    p = (PgHdr1 *)&((u8 *)pPg)[pCache->szPage];
 #endif
     if( benignMalloc ){ sqlite3EndBenignMalloc(); }
 #ifdef SQLITE_ENABLE_MEMORY_MANAGEMENT
     pcache1EnterMutex(pCache->pGroup);
 #endif
     if( pPg==0 ) return 0;
+#ifndef SQLITE_PCACHE_SEPARATE_HEADER
+    p = (PgHdr1 *)&((u8 *)pPg)[pCache->szPage];
+#endif
     p->page.pBuf = pPg;
     p->page.pExtra = &p[1];
     p->isBulkLocal = 0;
@@ -49420,6 +50237,7 @@ static sqlite3_pcache *pcache1Create(int szPage, int szExtra, int bPurgeable){
     }else{
       pGroup = &pcache1.grp;
     }
+    pcache1EnterMutex(pGroup);
     if( pGroup->lru.isAnchor==0 ){
       pGroup->lru.isAnchor = 1;
       pGroup->lru.pLruPrev = pGroup->lru.pLruNext = &pGroup->lru;
@@ -49429,7 +50247,6 @@ static sqlite3_pcache *pcache1Create(int szPage, int szExtra, int bPurgeable){
     pCache->szExtra = szExtra;
     pCache->szAlloc = szPage + szExtra + ROUND8(sizeof(PgHdr1));
     pCache->bPurgeable = (bPurgeable ? 1 : 0);
-    pcache1EnterMutex(pGroup);
     pcache1ResizeHash(pCache);
     if( bPurgeable ){
       pCache->nMin = 10;
@@ -51750,6 +52567,7 @@ static int pagerUnlockDb(Pager *pPager, int eLock){
     }
     IOTRACE(("UNLOCK %p %d\n", pPager, eLock))
   }
+  pPager->changeCountDone = pPager->tempFile; /* ticket fb3b3024ea238d5c */
   return rc;
 }
 
@@ -51937,6 +52755,7 @@ static int readMasterJournal(sqlite3_file *pJrnl, char *zMaster, u32 nMaster){
     len = 0;
   }
   zMaster[len] = '\0';
+  zMaster[len+1] = '\0';
    
   return SQLITE_OK;
 }
@@ -52470,7 +53289,6 @@ static void pager_unlock(Pager *pPager){
     ** code is cleared and the cache reset in the block below.
     */
     assert( pPager->errCode || pPager->eState!=PAGER_ERROR );
-    pPager->changeCountDone = 0;
     pPager->eState = PAGER_OPEN;
   }
 
@@ -52734,7 +53552,6 @@ static int pager_end_transaction(Pager *pPager, int hasMaster, int bCommit){
    && (!pagerUseWal(pPager) || sqlite3WalExclusiveMode(pPager->pWal, 0))
   ){
     rc2 = pagerUnlockDb(pPager, SHARED_LOCK);
-    pPager->changeCountDone = 0;
   }
   pPager->eState = PAGER_READER;
   pPager->setMaster = 0;
@@ -53173,15 +53990,16 @@ static int pager_delmaster(Pager *pPager, const char *zMaster){
   rc = sqlite3OsFileSize(pMaster, &nMasterJournal);
   if( rc!=SQLITE_OK ) goto delmaster_out;
   nMasterPtr = pVfs->mxPathname+1;
-  zMasterJournal = sqlite3Malloc(nMasterJournal + nMasterPtr + 1);
+  zMasterJournal = sqlite3Malloc(nMasterJournal + nMasterPtr + 2);
   if( !zMasterJournal ){
     rc = SQLITE_NOMEM_BKPT;
     goto delmaster_out;
   }
-  zMasterPtr = &zMasterJournal[nMasterJournal+1];
+  zMasterPtr = &zMasterJournal[nMasterJournal+2];
   rc = sqlite3OsRead(pMaster, zMasterJournal, (int)nMasterJournal, 0);
   if( rc!=SQLITE_OK ) goto delmaster_out;
   zMasterJournal[nMasterJournal] = 0;
+  zMasterJournal[nMasterJournal+1] = 0;
 
   zJournal = zMasterJournal;
   while( (zJournal-zMasterJournal)<nMasterJournal ){
@@ -55338,7 +56156,8 @@ SQLITE_PRIVATE int sqlite3PagerOpen(
   int pcacheSize = sqlite3PcacheSize();       /* Bytes to allocate for PCache */
   u32 szPageDflt = SQLITE_DEFAULT_PAGE_SIZE;  /* Default page size */
   const char *zUri = 0;    /* URI args to copy */
-  int nUri = 0;            /* Number of bytes of URI args at *zUri */
+  int nUriByte = 1;        /* Number of bytes of URI args at *zUri */
+  int nUri = 0;            /* Number of URI parameters */
 
   /* Figure out how much space is required for each journal file-handle
   ** (there are two of them, the main journal and the sub-journal).  */
@@ -55372,14 +56191,24 @@ SQLITE_PRIVATE int sqlite3PagerOpen(
     }
     zPathname[0] = 0; /* Make sure initialized even if FullPathname() fails */
     rc = sqlite3OsFullPathname(pVfs, zFilename, nPathname, zPathname);
+    if( rc!=SQLITE_OK ){
+      if( rc==SQLITE_OK_SYMLINK ){
+        if( vfsFlags & SQLITE_OPEN_NOFOLLOW ){
+          rc = SQLITE_CANTOPEN_SYMLINK;
+        }else{
+          rc = SQLITE_OK;
+        }
+      }
+    }
     nPathname = sqlite3Strlen30(zPathname);
     z = zUri = &zFilename[sqlite3Strlen30(zFilename)+1];
     while( *z ){
-      z += sqlite3Strlen30(z)+1;
-      z += sqlite3Strlen30(z)+1;
+      z += strlen(z)+1;
+      z += strlen(z)+1;
+      nUri++;
     }
-    nUri = (int)(&z[1] - zUri);
-    assert( nUri>=0 );
+    nUriByte = (int)(&z[1] - zUri);
+    assert( nUriByte>=1 );
     if( rc==SQLITE_OK && nPathname+8>pVfs->mxPathname ){
       /* This branch is taken when the journal path required by
       ** the database being opened will be more than pVfs->mxPathname
@@ -55404,50 +56233,103 @@ SQLITE_PRIVATE int sqlite3PagerOpen(
   **     Database file handle            (pVfs->szOsFile bytes)
   **     Sub-journal file handle         (journalFileSize bytes)
   **     Main journal file handle        (journalFileSize bytes)
+  **     \0\0\0\0 database prefix        (4 bytes)
   **     Database file name              (nPathname+1 bytes)
-  **     Journal file name               (nPathname+8+1 bytes)
+  **     URI query parameters            (nUriByte bytes)
+  **     Journal filename                (nPathname+8+1 bytes)
+  **     WAL filename                    (nPathname+4+1 bytes)
+  **     \0\0\0 terminator               (3 bytes)
+  **
+  ** Some 3rd-party software, over which we have no control, depends on
+  ** the specific order of the filenames and the \0 separators between them
+  ** so that it can (for example) find the database filename given the WAL
+  ** filename without using the sqlite3_filename_database() API.  This is a
+  ** misuse of SQLite and a bug in the 3rd-party software, but the 3rd-party
+  ** software is in widespread use, so we try to avoid changing the filename
+  ** order and formatting if possible.  In particular, the details of the
+  ** filename format expected by 3rd-party software should be as follows:
+  **
+  **   - Main Database Path
+  **   - \0
+  **   - Multiple URI components consisting of:
+  **     - Key
+  **     - \0
+  **     - Value
+  **     - \0
+  **   - \0
+  **   - Journal Path
+  **   - \0
+  **   - WAL Path (zWALName)
+  **   - \0
   */
   pPtr = (u8 *)sqlite3MallocZero(
-    ROUND8(sizeof(*pPager)) +      /* Pager structure */
-    ROUND8(pcacheSize) +           /* PCache object */
-    ROUND8(pVfs->szOsFile) +       /* The main db file */
-    journalFileSize * 2 +          /* The two journal files */ 
-    nPathname + 1 + nUri +         /* zFilename */
-    nPathname + 8 + 2              /* zJournal */
+    ROUND8(sizeof(*pPager)) +            /* Pager structure */
+    ROUND8(pcacheSize) +                 /* PCache object */
+    ROUND8(pVfs->szOsFile) +             /* The main db file */
+    journalFileSize * 2 +                /* The two journal files */
+    4 +                                  /* Database prefix */
+    nPathname + 1 +                      /* database filename */
+    nUriByte +                           /* query parameters */
+    nPathname + 8 + 1 +                  /* Journal filename */
 #ifndef SQLITE_OMIT_WAL
-    + nPathname + 4 + 2            /* zWal */
+    nPathname + 4 + 1 +                  /* WAL filename */
 #endif
+    3                                    /* Terminator */
   );
   assert( EIGHT_BYTE_ALIGNMENT(SQLITE_INT_TO_PTR(journalFileSize)) );
   if( !pPtr ){
     sqlite3DbFree(0, zPathname);
     return SQLITE_NOMEM_BKPT;
   }
-  pPager =              (Pager*)(pPtr);
-  pPager->pPCache =    (PCache*)(pPtr += ROUND8(sizeof(*pPager)));
-  pPager->fd =   (sqlite3_file*)(pPtr += ROUND8(pcacheSize));
-  pPager->sjfd = (sqlite3_file*)(pPtr += ROUND8(pVfs->szOsFile));
-  pPager->jfd =  (sqlite3_file*)(pPtr += journalFileSize);
-  pPager->zFilename =    (char*)(pPtr += journalFileSize);
+  pPager = (Pager*)pPtr;                  pPtr += ROUND8(sizeof(*pPager));
+  pPager->pPCache = (PCache*)pPtr;        pPtr += ROUND8(pcacheSize);
+  pPager->fd = (sqlite3_file*)pPtr;       pPtr += ROUND8(pVfs->szOsFile);
+  pPager->sjfd = (sqlite3_file*)pPtr;     pPtr += journalFileSize;
+  pPager->jfd =  (sqlite3_file*)pPtr;     pPtr += journalFileSize;
   assert( EIGHT_BYTE_ALIGNMENT(pPager->jfd) );
 
-  /* Fill in the Pager.zFilename and Pager.zJournal buffers, if required. */
-  if( zPathname ){
-    assert( nPathname>0 );
-    pPager->zJournal =   (char*)(pPtr += nPathname + 1 + nUri);
-    memcpy(pPager->zFilename, zPathname, nPathname);
-    if( nUri ) memcpy(&pPager->zFilename[nPathname+1], zUri, nUri);
-    memcpy(pPager->zJournal, zPathname, nPathname);
-    memcpy(&pPager->zJournal[nPathname], "-journal\000", 8+2);
-    sqlite3FileSuffix3(pPager->zFilename, pPager->zJournal);
+  /* Fill in the Pager.zFilename and pPager.zQueryParam fields */
+                                          pPtr += 4;  /* Skip zero prefix */
+  pPager->zFilename = (char*)pPtr;
+  if( nPathname>0 ){
+    memcpy(pPtr, zPathname, nPathname);   pPtr += nPathname + 1;
+    if( zUri ){
+      memcpy(pPtr, zUri, nUriByte);       pPtr += nUriByte;
+    }else{
+                                          pPtr++;
+    }
+  }
+
+
+  /* Fill in Pager.zJournal */
+  if( nPathname>0 ){
+    pPager->zJournal = (char*)pPtr;
+    memcpy(pPtr, zPathname, nPathname);   pPtr += nPathname;
+    memcpy(pPtr, "-journal",8);           pPtr += 8 + 1;
+#ifdef SQLITE_ENABLE_8_3_NAMES
+    sqlite3FileSuffix3(zFilename,pPager->zJournal);
+    pPtr = (u8*)(pPager->zJournal + sqlite3Strlen30(pPager->zJournal)+1);
+#endif
+  }else{
+    pPager->zJournal = 0;
+  }
+
 #ifndef SQLITE_OMIT_WAL
-    pPager->zWal = &pPager->zJournal[nPathname+8+1];
-    memcpy(pPager->zWal, zPathname, nPathname);
-    memcpy(&pPager->zWal[nPathname], "-wal\000", 4+1);
-    sqlite3FileSuffix3(pPager->zFilename, pPager->zWal);
+  /* Fill in Pager.zWal */
+  if( nPathname>0 ){
+    pPager->zWal = (char*)pPtr;
+    memcpy(pPtr, zPathname, nPathname);   pPtr += nPathname;
+    memcpy(pPtr, "-wal", 4);              pPtr += 4 + 1;
+#ifdef SQLITE_ENABLE_8_3_NAMES
+    sqlite3FileSuffix3(zFilename, pPager->zWal);
+    pPtr = (u8*)(pPager->zWal + sqlite3Strlen30(pPager->zWal)+1);
 #endif
-    sqlite3DbFree(0, zPathname);
+  }else{
+    pPager->zWal = 0;
   }
+#endif
+
+  if( nPathname ) sqlite3DbFree(0, zPathname);
   pPager->pVfs = pVfs;
   pPager->vfsFlags = vfsFlags;
 
@@ -55496,9 +56378,9 @@ SQLITE_PRIVATE int sqlite3PagerOpen(
         }
 #endif
       }
-      pPager->noLock = sqlite3_uri_boolean(zFilename, "nolock", 0);
+      pPager->noLock = sqlite3_uri_boolean(pPager->zFilename, "nolock", 0);
       if( (iDc & SQLITE_IOCAP_IMMUTABLE)!=0
-       || sqlite3_uri_boolean(zFilename, "immutable", 0) ){
+       || sqlite3_uri_boolean(pPager->zFilename, "immutable", 0) ){
           vfsFlags |= SQLITE_OPEN_READONLY;
           goto act_like_temp_file;
       }
@@ -57189,6 +58071,7 @@ SQLITE_PRIVATE int sqlite3PagerCommitPhaseTwo(Pager *pPager){
   ** But if (due to a coding error elsewhere in the system) it does get
   ** called, just return the same error code without doing anything. */
   if( NEVER(pPager->errCode) ) return pPager->errCode;
+  pPager->iDataVersion++;
 
   assert( pPager->eState==PAGER_WRITER_LOCKED
        || pPager->eState==PAGER_WRITER_FINISHED
@@ -57217,7 +58100,6 @@ SQLITE_PRIVATE int sqlite3PagerCommitPhaseTwo(Pager *pPager){
   }
 
   PAGERTRACE(("COMMIT %d\n", PAGERID(pPager)));
-  pPager->iDataVersion++;
   rc = pager_end_transaction(pPager, pPager->setMaster, 1);
   return pager_error(pPager, rc);
 }
@@ -57561,9 +58443,13 @@ SQLITE_PRIVATE int sqlite3PagerSavepoint(Pager *pPager, int op, int iSavepoint){
 ** behavior.  But when the Btree needs to know the filename for matching to
 ** shared cache, it uses nullIfMemDb==0 so that in-memory databases can
 ** participate in shared-cache.
+**
+** The return value to this routine is always safe to use with
+** sqlite3_uri_parameter() and sqlite3_filename_database() and friends.
 */
-SQLITE_PRIVATE const char *sqlite3PagerFilename(Pager *pPager, int nullIfMemDb){
-  return (nullIfMemDb && pPager->memDb) ? "" : pPager->zFilename;
+SQLITE_PRIVATE const char *sqlite3PagerFilename(const Pager *pPager, int nullIfMemDb){
+  static const char zFake[] = { 0x00, 0x00, 0x00, 0x00, 0x00, 0x00 };
+  return (nullIfMemDb && pPager->memDb) ? &zFake[4] : pPager->zFilename;
 }
 
 /*
@@ -58218,6 +59104,8 @@ SQLITE_PRIVATE int sqlite3PagerCloseWal(Pager *pPager, sqlite3 *db){
   return rc;
 }
 
+
+
 #ifdef SQLITE_ENABLE_SNAPSHOT
 /*
 ** If this is a WAL database, obtain a snapshot handle for the snapshot
@@ -60138,7 +61026,19 @@ static int walCheckpoint(
       ** not decreasing it. So assuming either that either the "old" or
       ** "new" version of the value is read, and not some arbitrary value
       ** that would never be written by a real client, things are still 
-      ** safe.  */
+      ** safe.
+      **
+      ** Astute readers have pointed out that the assumption stated in the
+      ** last sentence of the previous paragraph is not guaranteed to be
+      ** true for all conforming systems.  However, the assumption is true
+      ** for all compilers and architectures in common use today (circa
+      ** 2019-11-27) and the alternatives are both slow and complex, and
+      ** so we will continue to go with the current design for now.  If this
+      ** bothers you, or if you really are running on a system where aligned
+      ** 32-bit reads and writes are not atomic, then you can simply avoid
+      ** the use of WAL mode, or only use WAL mode together with
+      ** PRAGMA locking_mode=EXCLUSIVE and all will be well.
+      */
       u32 y = pInfo->aReadMark[i];
       if( mxSafeFrame>y ){
         assert( y<=pWal->hdr.mxFrame );
@@ -60215,6 +61115,10 @@ static int walCheckpoint(
             rc = sqlite3OsSync(pWal->pDbFd, CKPT_SYNC_FLAGS(sync_flags));
           }
         }
+        if( rc==SQLITE_OK ){
+          rc = sqlite3OsFileControl(pWal->pDbFd, SQLITE_FCNTL_CKPT_DONE, 0);
+          if( rc==SQLITE_NOTFOUND ) rc = SQLITE_OK;
+        }
         if( rc==SQLITE_OK ){
           pInfo->nBackfill = mxSafeFrame;
         }
@@ -61222,9 +62126,9 @@ SQLITE_PRIVATE int sqlite3WalFindFrame(
     }
     nCollide = HASHTABLE_NSLOT;
     for(iKey=walHash(pgno); sLoc.aHash[iKey]; iKey=walNextHash(iKey)){
-      u32 iFrame = sLoc.aHash[iKey] + sLoc.iZero;
-      if( iFrame<=iLast && iFrame>=pWal->minFrame
-       && sLoc.aPgno[sLoc.aHash[iKey]]==pgno ){
+      u32 iH = sLoc.aHash[iKey];
+      u32 iFrame = iH + sLoc.iZero;
+      if( iFrame<=iLast && iFrame>=pWal->minFrame && sLoc.aPgno[iH]==pgno ){
         assert( iFrame>iRead || CORRUPT_DB );
         iRead = iFrame;
       }
@@ -61792,6 +62696,7 @@ SQLITE_PRIVATE int sqlite3WalFrames(
         if( rc ) return rc;
         iOffset += szFrame;
         nExtra++;
+        assert( pLast!=0 );
       }
     }
     if( bSync ){
@@ -61824,6 +62729,7 @@ SQLITE_PRIVATE int sqlite3WalFrames(
     iFrame++;
     rc = walIndexAppend(pWal, iFrame, p->pgno);
   }
+  assert( pLast!=0 || nExtra==0 );
   while( rc==SQLITE_OK && nExtra>0 ){
     iFrame++;
     nExtra--;
@@ -62723,6 +63629,7 @@ struct BtCursor {
 #define BTCF_AtLast       0x08   /* Cursor is pointing ot the last entry */
 #define BTCF_Incrblob     0x10   /* True if an incremental I/O handle */
 #define BTCF_Multiple     0x20   /* Maybe another cursor on the same btree */
+#define BTCF_Pinned       0x40   /* Cursor is busy and cannot be moved */
 
 /*
 ** Potential values for BtCursor.eState.
@@ -62866,6 +63773,7 @@ struct IntegrityCk {
   int v1, v2;       /* Values for up to two %d fields in zPfx */
   StrAccum errMsg;  /* Accumulate the error message text here */
   u32 *heap;        /* Min-heap used for analyzing cell coverage */
+  sqlite3 *db;      /* Database connection running the check */
 };
 
 /*
@@ -63887,6 +64795,9 @@ static int saveCursorPosition(BtCursor *pCur){
   assert( 0==pCur->pKey );
   assert( cursorHoldsMutex(pCur) );
 
+  if( pCur->curFlags & BTCF_Pinned ){
+    return SQLITE_CONSTRAINT_PINNED;
+  }
   if( pCur->eState==CURSOR_SKIPNEXT ){
     pCur->eState = CURSOR_VALID;
   }else{
@@ -64634,7 +65545,7 @@ static int defragmentPage(MemPage *pPage, int nMaxFrag){
         int sz2 = 0;
         int sz = get2byte(&data[iFree+2]);
         int top = get2byte(&data[hdr+5]);
-        if( top>=iFree ){
+        if( NEVER(top>=iFree) ){
           return SQLITE_CORRUPT_PAGE(pPage);
         }
         if( iFree2 ){
@@ -64643,7 +65554,7 @@ static int defragmentPage(MemPage *pPage, int nMaxFrag){
           if( iFree2+sz2 > usableSize ) return SQLITE_CORRUPT_PAGE(pPage);
           memmove(&data[iFree+sz+sz2], &data[iFree+sz], iFree2-(iFree+sz));
           sz += sz2;
-        }else if( iFree+sz>usableSize ){
+        }else if( NEVER(iFree+sz>usableSize) ){
           return SQLITE_CORRUPT_PAGE(pPage);
         }
 
@@ -64816,7 +65727,7 @@ static int allocateSpace(MemPage *pPage, int nByte, int *pIdx){
   ** However, that integer is too large to be stored in a 2-byte unsigned
   ** integer, so a value of 0 is used in its place. */
   top = get2byte(&data[hdr+5]);
-  assert( top<=(int)pPage->pBt->usableSize ); /* Prevent by getAndInitPage() */
+  assert( top<=(int)pPage->pBt->usableSize ); /* by btreeComputeFreeSpace() */
   if( gap>top ){
     if( top==0 && pPage->pBt->usableSize==65536 ){
       top = 65536;
@@ -64835,9 +65746,14 @@ static int allocateSpace(MemPage *pPage, int nByte, int *pIdx){
   if( (data[hdr+2] || data[hdr+1]) && gap+2<=top ){
     u8 *pSpace = pageFindSlot(pPage, nByte, &rc);
     if( pSpace ){
-      assert( pSpace>=data && (pSpace - data)<65536 );
-      *pIdx = (int)(pSpace - data);
-      return SQLITE_OK;
+      int g2;
+      assert( pSpace+nByte<=data+pPage->pBt->usableSize );
+      *pIdx = g2 = (int)(pSpace-data);
+      if( NEVER(g2<=gap) ){
+        return SQLITE_CORRUPT_PAGE(pPage);
+      }else{
+        return SQLITE_OK;
+      }
     }else if( rc ){
       return rc;
     }
@@ -64911,12 +65827,12 @@ static int freeSpace(MemPage *pPage, u16 iStart, u16 iSize){
   }else{
     while( (iFreeBlk = get2byte(&data[iPtr]))<iStart ){
       if( iFreeBlk<iPtr+4 ){
-        if( iFreeBlk==0 ) break;
+        if( iFreeBlk==0 ) break; /* TH3: corrupt082.100 */
         return SQLITE_CORRUPT_PAGE(pPage);
       }
       iPtr = iFreeBlk;
     }
-    if( iFreeBlk>pPage->pBt->usableSize-4 ){
+    if( iFreeBlk>pPage->pBt->usableSize-4 ){ /* TH3: corrupt081.100 */
       return SQLITE_CORRUPT_PAGE(pPage);
     }
     assert( iFreeBlk>iPtr || iFreeBlk==0 );
@@ -64931,7 +65847,7 @@ static int freeSpace(MemPage *pPage, u16 iStart, u16 iSize){
       nFrag = iFreeBlk - iEnd;
       if( iEnd>iFreeBlk ) return SQLITE_CORRUPT_PAGE(pPage);
       iEnd = iFreeBlk + get2byte(&data[iFreeBlk+2]);
-      if( iEnd > pPage->pBt->usableSize ){
+      if( NEVER(iEnd > pPage->pBt->usableSize) ){
         return SQLITE_CORRUPT_PAGE(pPage);
       }
       iSize = iEnd - iStart;
@@ -64959,7 +65875,8 @@ static int freeSpace(MemPage *pPage, u16 iStart, u16 iSize){
     /* The new freeblock is at the beginning of the cell content area,
     ** so just extend the cell content area rather than create another
     ** freelist entry */
-    if( iStart<x || iPtr!=hdr+1 ) return SQLITE_CORRUPT_PAGE(pPage);
+    if( iStart<x ) return SQLITE_CORRUPT_PAGE(pPage);
+    if( NEVER(iPtr!=hdr+1) ) return SQLITE_CORRUPT_PAGE(pPage);
     put2byte(&data[hdr+1], iFreeBlk);
     put2byte(&data[hdr+5], iEnd);
   }else{
@@ -65079,7 +65996,7 @@ static int btreeComputeFreeSpace(MemPage *pPage){
   nFree = data[hdr+7] + top;  /* Init nFree to non-freeblock free space */
   if( pc>0 ){
     u32 next, size;
-    if( pc<iCellFirst ){
+    if( pc<top ){
       /* EVIDENCE-OF: R-55530-52930 In a well-formed b-tree page, there will
       ** always be at least one cell before the first freeblock.
       */
@@ -65113,7 +66030,7 @@ static int btreeComputeFreeSpace(MemPage *pPage){
   ** serves to verify that the offset to the start of the cell-content
   ** area, according to the page header, lies within the page.
   */
-  if( nFree>usableSize ){
+  if( nFree>usableSize || nFree<iCellFirst ){
     return SQLITE_CORRUPT_PAGE(pPage);
   }
   pPage->nFree = (u16)(nFree - iCellFirst);
@@ -65316,12 +66233,12 @@ static MemPage *btreePageLookup(BtShared *pBt, Pgno pgno){
 ** error, return ((unsigned int)-1).
 */
 static Pgno btreePagecount(BtShared *pBt){
+  assert( (pBt->nPage & 0x80000000)==0 || CORRUPT_DB );
   return pBt->nPage;
 }
 SQLITE_PRIVATE u32 sqlite3BtreeLastPage(Btree *p){
   assert( sqlite3BtreeHoldsMutex(p) );
-  assert( ((p->pBt->nPage)&0x80000000)==0 );
-  return btreePagecount(p->pBt);
+  return btreePagecount(p->pBt) & 0x7fffffff;
 }
 
 /*
@@ -65588,9 +66505,13 @@ SQLITE_PRIVATE int sqlite3BtreeOpen(
         rc = sqlite3OsFullPathname(pVfs, zFilename,
                                    nFullPathname, zFullPathname);
         if( rc ){
-          sqlite3_free(zFullPathname);
-          sqlite3_free(p);
-          return rc;
+          if( rc==SQLITE_OK_SYMLINK ){
+            rc = SQLITE_OK;
+          }else{
+            sqlite3_free(zFullPathname);
+            sqlite3_free(p);
+            return rc;
+          }
         }
       }
 #if SQLITE_THREADSAFE
@@ -67341,6 +68262,18 @@ SQLITE_PRIVATE int sqlite3BtreeTripAllCursors(Btree *pBtree, int errCode, int wr
   return rc;
 }
 
+/*
+** Set the pBt->nPage field correctly, according to the current
+** state of the database.  Assume pBt->pPage1 is valid.
+*/
+static void btreeSetNPage(BtShared *pBt, MemPage *pPage1){
+  int nPage = get4byte(&pPage1->aData[28]);
+  testcase( nPage==0 );
+  if( nPage==0 ) sqlite3PagerPagecount(pBt->pPager, &nPage);
+  testcase( pBt->nPage!=nPage );
+  pBt->nPage = nPage;
+}
+
 /*
 ** Rollback the transaction in progress.
 **
@@ -67386,11 +68319,7 @@ SQLITE_PRIVATE int sqlite3BtreeRollback(Btree *p, int tripCode, int writeOnly){
     ** call btreeGetPage() on page 1 again to make
     ** sure pPage1->aData is set correctly. */
     if( btreeGetPage(pBt, 1, &pPage1, 0)==SQLITE_OK ){
-      int nPage = get4byte(28+(u8*)pPage1->aData);
-      testcase( nPage==0 );
-      if( nPage==0 ) sqlite3PagerPagecount(pBt->pPager, &nPage);
-      testcase( pBt->nPage!=nPage );
-      pBt->nPage = nPage;
+      btreeSetNPage(pBt, pPage1);
       releasePageOne(pPage1);
     }
     assert( countValidCursors(pBt, 1)==0 );
@@ -67470,12 +68399,11 @@ SQLITE_PRIVATE int sqlite3BtreeSavepoint(Btree *p, int op, int iSavepoint){
         pBt->nPage = 0;
       }
       rc = newDatabase(pBt);
-      pBt->nPage = get4byte(28 + pBt->pPage1->aData);
+      btreeSetNPage(pBt, pBt->pPage1);
 
-      /* The database size was written into the offset 28 of the header
-      ** when the transaction started, so we know that the value at offset
-      ** 28 is nonzero. */
-      assert( pBt->nPage>0 );
+      /* pBt->nPage might be zero if the database was corrupt when 
+      ** the transaction was started. Otherwise, it must be at least 1.  */
+      assert( CORRUPT_DB || pBt->nPage>0 );
     }
     sqlite3BtreeLeave(p);
   }
@@ -67543,8 +68471,9 @@ static int btreeCursor(
   /* The following assert statements verify that if this is a sharable 
   ** b-tree database, the connection is holding the required table locks, 
   ** and that no other connection has any open cursor that conflicts with 
-  ** this lock.  */
-  assert( hasSharedCacheTableLock(p, iTable, pKeyInfo!=0, (wrFlag?2:1)) );
+  ** this lock.  The iTable<1 term disables the check for corrupt schemas. */
+  assert( hasSharedCacheTableLock(p, iTable, pKeyInfo!=0, (wrFlag?2:1))
+          || iTable<1 );
   assert( wrFlag==0 || !hasReadConflicts(p, iTable) );
 
   /* Assert that the caller has opened the required transaction. */
@@ -67557,9 +68486,13 @@ static int btreeCursor(
     allocateTempSpace(pBt);
     if( pBt->pTmpSpace==0 ) return SQLITE_NOMEM_BKPT;
   }
-  if( iTable==1 && btreePagecount(pBt)==0 ){
-    assert( wrFlag==0 );
-    iTable = 0;
+  if( iTable<=1 ){
+    if( iTable<1 ){
+      return SQLITE_CORRUPT_BKPT;
+    }else if( btreePagecount(pBt)==0 ){
+      assert( wrFlag==0 );
+      iTable = 0;
+    }
   }
 
   /* Now that no other errors can occur, finish filling in the BtCursor
@@ -67584,6 +68517,19 @@ static int btreeCursor(
   pCur->eState = CURSOR_INVALID;
   return SQLITE_OK;
 }
+static int btreeCursorWithLock(
+  Btree *p,                              /* The btree */
+  int iTable,                            /* Root page of table to open */
+  int wrFlag,                            /* 1 to write. 0 read-only */
+  struct KeyInfo *pKeyInfo,              /* First arg to comparison function */
+  BtCursor *pCur                         /* Space for new cursor */
+){
+  int rc;
+  sqlite3BtreeEnter(p);
+  rc = btreeCursor(p, iTable, wrFlag, pKeyInfo, pCur);
+  sqlite3BtreeLeave(p);
+  return rc;
+}
 SQLITE_PRIVATE int sqlite3BtreeCursor(
   Btree *p,                                   /* The btree */
   int iTable,                                 /* Root page of table to open */
@@ -67591,15 +68537,11 @@ SQLITE_PRIVATE int sqlite3BtreeCursor(
   struct KeyInfo *pKeyInfo,                   /* First arg to xCompare() */
   BtCursor *pCur                              /* Write new cursor here */
 ){
-  int rc;
-  if( iTable<1 ){
-    rc = SQLITE_CORRUPT_BKPT;
+  if( p->sharable ){
+    return btreeCursorWithLock(p, iTable, wrFlag, pKeyInfo, pCur);
   }else{
-    sqlite3BtreeEnter(p);
-    rc = btreeCursor(p, iTable, wrFlag, pKeyInfo, pCur);
-    sqlite3BtreeLeave(p);
+    return btreeCursor(p, iTable, wrFlag, pKeyInfo, pCur);
   }
-  return rc;
 }
 
 /*
@@ -67722,6 +68664,18 @@ SQLITE_PRIVATE i64 sqlite3BtreeIntegerKey(BtCursor *pCur){
   return pCur->info.nKey;
 }
 
+/*
+** Pin or unpin a cursor.
+*/
+SQLITE_PRIVATE void sqlite3BtreeCursorPin(BtCursor *pCur){
+  assert( (pCur->curFlags & BTCF_Pinned)==0 );
+  pCur->curFlags |= BTCF_Pinned;
+}
+SQLITE_PRIVATE void sqlite3BtreeCursorUnpin(BtCursor *pCur){
+  assert( (pCur->curFlags & BTCF_Pinned)!=0 );
+  pCur->curFlags &= ~BTCF_Pinned;
+}
+
 #ifdef SQLITE_ENABLE_OFFSET_SQL_FUNC
 /*
 ** Return the offset into the database file for the start of the
@@ -68057,6 +69011,7 @@ static int accessPayload(
           assert( aWrite>=pBufStart );                         /* due to (6) */
           memcpy(aSave, aWrite, 4);
           rc = sqlite3OsRead(fd, aWrite, a+4, (i64)pBt->pageSize*(nextPage-1));
+          if( rc && nextPage>pBt->nPage ) rc = SQLITE_CORRUPT_BKPT;
           nextPage = get4byte(aWrite);
           memcpy(aWrite, aSave, 4);
         }else
@@ -68483,6 +69438,7 @@ SQLITE_PRIVATE int sqlite3BtreeLast(BtCursor *pCur, int *pRes){
     assert( pCur->ix==pCur->pPage->nCell-1 );
     assert( pCur->pPage->leaf );
 #endif
+    *pRes = 0;
     return SQLITE_OK;
   }
 
@@ -68704,6 +69660,7 @@ SQLITE_PRIVATE int sqlite3BtreeMovetoUnpacked(
           ** case this happens.  */
           void *pCellKey;
           u8 * const pCellBody = pCell - pPage->childPtrSize;
+          const int nOverrun = 18;  /* Size of the overrun padding */
           pPage->xParseCell(pPage, pCellBody, &pCur->info);
           nCell = (int)pCur->info.nKey;
           testcase( nCell<0 );   /* True if key size is 2^32 or more */
@@ -68714,13 +69671,14 @@ SQLITE_PRIVATE int sqlite3BtreeMovetoUnpacked(
             rc = SQLITE_CORRUPT_PAGE(pPage);
             goto moveto_finish;
           }
-          pCellKey = sqlite3Malloc( nCell+18 );
+          pCellKey = sqlite3Malloc( nCell+nOverrun );
           if( pCellKey==0 ){
             rc = SQLITE_NOMEM_BKPT;
             goto moveto_finish;
           }
           pCur->ix = (u16)idx;
           rc = accessPayload(pCur, 0, nCell, (unsigned char*)pCellKey, 0);
+          memset(((u8*)pCellKey)+nCell,0,nOverrun); /* Fix uninit warnings */
           pCur->curFlags &= ~BTCF_ValidOvfl;
           if( rc ){
             sqlite3_free(pCellKey);
@@ -68874,8 +69832,11 @@ static SQLITE_NOINLINE int btreeNext(BtCursor *pCur){
   ** to be invalid here. This can only occur if a second cursor modifies
   ** the page while cursor pCur is holding a reference to it. Which can
   ** only happen if the database is corrupt in such a way as to link the
-  ** page into more than one b-tree structure. */
-  testcase( idx>pPage->nCell );
+  ** page into more than one b-tree structure.
+  **
+  ** Update 2019-12-23: appears to long longer be possible after the
+  ** addition of anotherValidCursor() condition on balance_deeper().  */
+  harmless( idx>pPage->nCell );
 
   if( idx>=pPage->nCell ){
     if( !pPage->leaf ){
@@ -69842,12 +70803,7 @@ static void insertCell(
   assert( pPage->nOverflow<=ArraySize(pPage->apOvfl) );
   assert( ArraySize(pPage->apOvfl)==ArraySize(pPage->aiOvfl) );
   assert( sqlite3_mutex_held(pPage->pBt->mutex) );
-  /* The cell should normally be sized correctly.  However, when moving a
-  ** malformed cell from a leaf page to an interior page, if the cell size
-  ** wanted to be less than 4 but got rounded up to 4 on the leaf, then size
-  ** might be less than 8 (leaf-size + pointer) on the interior node.  Hence
-  ** the term after the || in the following assert(). */
-  assert( sz==pPage->xCellSize(pPage, pCell) || (sz==8 && iChild>0) );
+  assert( sz==pPage->xCellSize(pPage, pCell) || CORRUPT_DB );
   assert( pPage->nFree>=0 );
   if( pPage->nOverflow || sz+2>pPage->nFree ){
     if( pTemp ){
@@ -70105,7 +71061,7 @@ static int rebuildPage(
     if( pData < pCellptr ) return SQLITE_CORRUPT_BKPT;
     memcpy(pData, pCell, sz);
     assert( sz==pPg->xCellSize(pPg, pCell) || CORRUPT_DB );
-    testcase( sz!=pPg->xCellSize(pPg,pCell) );
+    testcase( sz!=pPg->xCellSize(pPg,pCell) )
     i++;
     if( i>=iEnd ) break;
     if( pCArray->ixNx[k]<=i ){
@@ -70171,7 +71127,8 @@ static int pageInsertArray(
   while( 1 /*Exit by break*/ ){
     int sz, rc;
     u8 *pSlot;
-    sz = cachedCellSize(pCArray, i);
+    assert( pCArray->szCell[i]!=0 );
+    sz = pCArray->szCell[i];
     if( (aData[1]==0 && aData[2]==0) || (pSlot = pageFindSlot(pPg,sz,&rc))==0 ){
       if( (pData - pBegin)<sz ) return 1;
       pData -= sz;
@@ -70332,6 +71289,7 @@ static int editPage(
         memmove(&pCellptr[2], pCellptr, (nCell - iCell) * 2);
       }
       nCell++;
+      cachedCellSize(pCArray, iCell+iNew);
       if( pageInsertArray(
             pPg, pBegin, &pData, pCellptr,
             iCell+iNew, 1, pCArray
@@ -70825,6 +71783,7 @@ static int balance_nonroot(
     u16 maskPage = pOld->maskPage;
     u8 *piCell = aData + pOld->cellOffset;
     u8 *piEnd;
+    VVA_ONLY( int nCellAtStart = b.nCell; )
 
     /* Verify that all sibling pages are of the same "type" (table-leaf,
     ** table-interior, index-leaf, or index-interior).
@@ -70853,6 +71812,10 @@ static int balance_nonroot(
     */
     memset(&b.szCell[b.nCell], 0, sizeof(b.szCell[0])*(limit+pOld->nOverflow));
     if( pOld->nOverflow>0 ){
+      if( NEVER(limit<pOld->aiOvfl[0]) ){
+        rc = SQLITE_CORRUPT_BKPT;
+        goto balance_cleanup;
+      }
       limit = pOld->aiOvfl[0];
       for(j=0; j<limit; j++){
         b.apCell[b.nCell] = aData + (maskPage & get2byteAligned(piCell));
@@ -70872,6 +71835,7 @@ static int balance_nonroot(
       piCell += 2;
       b.nCell++;
     }
+    assert( (b.nCell-nCellAtStart)==(pOld->nCell+pOld->nOverflow) );
 
     cntOld[i] = b.nCell;
     if( i<nOld-1 && !leafData){
@@ -71134,6 +72098,8 @@ static int balance_nonroot(
   ));
 
   assert( sqlite3PagerIswriteable(pParent->pDbPage) );
+  assert( nNew>=1 && nNew<=ArraySize(apNew) );
+  assert( apNew[nNew-1]!=0 );
   put4byte(pRight, apNew[nNew-1]->pgno);
 
   /* If the sibling pages are not leaves, ensure that the right-child pointer
@@ -71172,6 +72138,7 @@ static int balance_nonroot(
       while( i==cntOldNext ){
         iOld++;
         assert( iOld<nNew || iOld<nOld );
+        assert( iOld>=0 && iOld<NB );
         pOld = iOld<nNew ? apNew[iOld] : apOld[iOld];
         cntOldNext += pOld->nCell + pOld->nOverflow + !leafData;
       }
@@ -71458,6 +72425,30 @@ static int balance_deeper(MemPage *pRoot, MemPage **ppChild){
   return SQLITE_OK;
 }
 
+/*
+** Return SQLITE_CORRUPT if any cursor other than pCur is currently valid
+** on the same B-tree as pCur.
+**
+** This can if a database is corrupt with two or more SQL tables
+** pointing to the same b-tree.  If an insert occurs on one SQL table
+** and causes a BEFORE TRIGGER to do a secondary insert on the other SQL
+** table linked to the same b-tree.  If the secondary insert causes a
+** rebalance, that can change content out from under the cursor on the
+** first SQL table, violating invariants on the first insert.
+*/
+static int anotherValidCursor(BtCursor *pCur){
+  BtCursor *pOther;
+  for(pOther=pCur->pBt->pCursor; pOther; pOther=pOther->pNext){
+    if( pOther!=pCur
+     && pOther->eState==CURSOR_VALID
+     && pOther->pPage==pCur->pPage
+    ){
+      return SQLITE_CORRUPT_BKPT;
+    }
+  }
+  return SQLITE_OK;
+}
+
 /*
 ** The page that pCur currently points to has just been modified in
 ** some way. This function figures out if this modification means the
@@ -71478,12 +72469,14 @@ static int balance(BtCursor *pCur){
   VVA_ONLY( int balance_deeper_called = 0 );
 
   do {
-    int iPage = pCur->iPage;
+    int iPage;
     MemPage *pPage = pCur->pPage;
 
     if( NEVER(pPage->nFree<0) && btreeComputeFreeSpace(pPage) ) break;
-    if( iPage==0 ){
-      if( pPage->nOverflow ){
+    if( pPage->nOverflow==0 && pPage->nFree<=nMin ){
+      break;
+    }else if( (iPage = pCur->iPage)==0 ){
+      if( pPage->nOverflow && (rc = anotherValidCursor(pCur))==SQLITE_OK ){
         /* The root page of the b-tree is overfull. In this case call the
         ** balance_deeper() function to create a new child for the root-page
         ** and copy the current contents of the root-page to it. The
@@ -71503,8 +72496,6 @@ static int balance(BtCursor *pCur){
       }else{
         break;
       }
-    }else if( pPage->nOverflow==0 && pPage->nFree<=nMin ){
-      break;
     }else{
       MemPage * const pParent = pCur->apPage[iPage-1];
       int const iIdx = pCur->aiIdx[iPage-1];
@@ -71646,7 +72637,9 @@ static int btreeOverwriteCell(BtCursor *pCur, const BtreePayload *pX){
   Pgno ovflPgno;                      /* Next overflow page to write */
   u32 ovflPageSize;                   /* Size to write on overflow page */
 
-  if( pCur->info.pPayload + pCur->info.nLocal > pPage->aDataEnd ){
+  if( pCur->info.pPayload + pCur->info.nLocal > pPage->aDataEnd
+   || pCur->info.pPayload < pPage->aData + pPage->cellOffset
+  ){
     return SQLITE_CORRUPT_BKPT;
   }
   /* Overwrite the local portion first */
@@ -71779,7 +72772,6 @@ SQLITE_PRIVATE int sqlite3BtreeInsert(
     if( flags & BTREE_SAVEPOSITION ){
       assert( pCur->curFlags & BTCF_ValidNKey );
       assert( pX->nKey==pCur->info.nKey );
-      assert( pCur->info.nSize!=0 );
       assert( loc==0 );
     }
 #endif
@@ -71854,7 +72846,9 @@ SQLITE_PRIVATE int sqlite3BtreeInsert(
     }
 
   }
-  assert( pCur->eState==CURSOR_VALID || (pCur->eState==CURSOR_INVALID && loc) );
+  assert( pCur->eState==CURSOR_VALID 
+       || (pCur->eState==CURSOR_INVALID && loc)
+       || CORRUPT_DB );
 
   pPage = pCur->pPage;
   assert( pPage->intKey || pX->nKey>=0 );
@@ -71887,6 +72881,8 @@ SQLITE_PRIVATE int sqlite3BtreeInsert(
       memcpy(newCell, oldCell, 4);
     }
     rc = clearCell(pPage, oldCell, &info);
+    testcase( pCur->curFlags & BTCF_ValidOvfl );
+    invalidateOverflowCache(pCur);
     if( info.nSize==szNew && info.nLocal==info.nPayload 
      && (!ISAUTOVACUUM || szNew<pPage->minLocal)
     ){
@@ -71900,7 +72896,12 @@ SQLITE_PRIVATE int sqlite3BtreeInsert(
       ** new entry uses overflow pages, as the insertCell() call below is
       ** necessary to add the PTRMAP_OVERFLOW1 pointer-map entry.  */
       assert( rc==SQLITE_OK ); /* clearCell never fails when nLocal==nPayload */
-      if( oldCell+szNew > pPage->aDataEnd ) return SQLITE_CORRUPT_BKPT;
+      if( oldCell < pPage->aData+pPage->hdrOffset+10 ){
+        return SQLITE_CORRUPT_BKPT;
+      }
+      if( oldCell+szNew > pPage->aDataEnd ){
+        return SQLITE_CORRUPT_BKPT;
+      }
       memcpy(oldCell, newCell, szNew);
       return SQLITE_OK;
     }
@@ -72618,7 +73619,7 @@ SQLITE_PRIVATE int sqlite3BtreeUpdateMeta(Btree *p, int idx, u32 iMeta){
 ** Otherwise, if an error is encountered (i.e. an IO error or database
 ** corruption) an SQLite error code is returned.
 */
-SQLITE_PRIVATE int sqlite3BtreeCount(BtCursor *pCur, i64 *pnEntry){
+SQLITE_PRIVATE int sqlite3BtreeCount(sqlite3 *db, BtCursor *pCur, i64 *pnEntry){
   i64 nEntry = 0;                      /* Value to return in *pnEntry */
   int rc;                              /* Return code */
 
@@ -72631,7 +73632,7 @@ SQLITE_PRIVATE int sqlite3BtreeCount(BtCursor *pCur, i64 *pnEntry){
   /* Unless an error occurs, the following loop runs one iteration for each
   ** page in the B-Tree structure (not including overflow pages). 
   */
-  while( rc==SQLITE_OK ){
+  while( rc==SQLITE_OK && !db->u1.isInterrupted ){
     int iIdx;                          /* Index of child node in parent */
     MemPage *pPage;                    /* Current page of the b-tree */
 
@@ -72757,6 +73758,7 @@ static int checkRef(IntegrityCk *pCheck, Pgno iPage){
     checkAppendMsg(pCheck, "2nd reference to page %d", iPage);
     return 1;
   }
+  if( pCheck->db->u1.isInterrupted ) return 1;
   setPageReferenced(pCheck, iPage);
   return 0;
 }
@@ -73200,6 +74202,7 @@ end_of_check:
 ** returned.  If a memory allocation error occurs, NULL is returned.
 */
 SQLITE_PRIVATE char *sqlite3BtreeIntegrityCheck(
+  sqlite3 *db,  /* Database connection that is running the check */
   Btree *p,     /* The btree to be checked */
   int *aRoot,   /* An array of root pages numbers for individual trees */
   int nRoot,    /* Number of entries in aRoot[] */
@@ -73217,6 +74220,7 @@ SQLITE_PRIVATE char *sqlite3BtreeIntegrityCheck(
   assert( p->inTrans>TRANS_NONE && pBt->inTransaction>TRANS_NONE );
   VVA_ONLY( nRef = sqlite3PagerRefcount(pBt->pPager) );
   assert( nRef>=0 );
+  sCheck.db = db;
   sCheck.pBt = pBt;
   sCheck.pPager = pBt->pPager;
   sCheck.nPage = btreePagecount(sCheck.pBt);
@@ -73892,7 +74896,7 @@ static int backupOnePage(
   if( nSrcReserve!=nDestReserve ){
     u32 newPgsz = nSrcPgsz;
     rc = sqlite3PagerSetPagesize(pDestPager, &newPgsz, nSrcReserve);
-    if( rc==SQLITE_OK && newPgsz!=nSrcPgsz ) rc = SQLITE_READONLY;
+    if( rc==SQLITE_OK && newPgsz!=(u32)nSrcPgsz ) rc = SQLITE_READONLY;
   }
 #endif
 
@@ -74237,8 +75241,10 @@ SQLITE_API int sqlite3_backup_finish(sqlite3_backup *p){
   }
   if( p->isAttached ){
     pp = sqlite3PagerBackupPtr(sqlite3BtreePager(p->pSrc));
+    assert( pp!=0 );
     while( *pp!=p ){
       pp = &(*pp)->pNext;
+      assert( pp!=0 );
     }
     *pp = p->pNext;
   }
@@ -74439,6 +75445,11 @@ copy_finished:
 /* #include "sqliteInt.h" */
 /* #include "vdbeInt.h" */
 
+/* True if X is a power of two.  0 is considered a power of two here.
+** In other words, return true if X has at most one bit set.
+*/
+#define ISPOWEROF2(X)  (((X)&((X)-1))==0)
+
 #ifdef SQLITE_DEBUG
 /*
 ** Check invariants on a Mem object.
@@ -74458,8 +75469,8 @@ SQLITE_PRIVATE int sqlite3VdbeCheckMemInvariants(Mem *p){
   ** That saves a few cycles in inner loops. */
   assert( (p->flags & MEM_Dyn)==0 || p->szMalloc==0 );
 
-  /* Cannot be both MEM_Int and MEM_Real at the same time */
-  assert( (p->flags & (MEM_Int|MEM_Real))!=(MEM_Int|MEM_Real) );
+  /* Cannot have more than one of MEM_Int, MEM_Real, or MEM_IntReal */
+  assert( ISPOWEROF2(p->flags & (MEM_Int|MEM_Real|MEM_IntReal)) );
 
   if( p->flags & MEM_Null ){
     /* Cannot be both MEM_Null and some other type */
@@ -74513,9 +75524,31 @@ SQLITE_PRIVATE int sqlite3VdbeCheckMemInvariants(Mem *p){
 }
 #endif
 
+/*
+** Render a Mem object which is one of MEM_Int, MEM_Real, or MEM_IntReal
+** into a buffer.
+*/
+static void vdbeMemRenderNum(int sz, char *zBuf, Mem *p){
+  StrAccum acc;
+  assert( p->flags & (MEM_Int|MEM_Real|MEM_IntReal) );
+  sqlite3StrAccumInit(&acc, 0, zBuf, sz, 0);
+  if( p->flags & MEM_Int ){
+    sqlite3_str_appendf(&acc, "%lld", p->u.i);
+  }else if( p->flags & MEM_IntReal ){
+    sqlite3_str_appendf(&acc, "%!.15g", (double)p->u.i);
+  }else{
+    sqlite3_str_appendf(&acc, "%!.15g", p->u.r);
+  }
+  assert( acc.zText==zBuf && acc.mxAlloc<=0 );
+  zBuf[acc.nChar] = 0; /* Fast version of sqlite3StrAccumFinish(&acc) */
+}
+
 #ifdef SQLITE_DEBUG
 /*
-** Check that string value of pMem agrees with its integer or real value.
+** Validity checks on pMem.  pMem holds a string.
+**
+** (1) Check that string value of pMem agrees with its integer or real value.
+** (2) Check that the string is correctly zero terminated
 **
 ** A single int or real value always converts to the same strings.  But
 ** many different strings can be converted into the same int or real.
@@ -74533,17 +75566,24 @@ SQLITE_PRIVATE int sqlite3VdbeCheckMemInvariants(Mem *p){
 **
 ** This routine is for use inside of assert() statements only.
 */
-SQLITE_PRIVATE int sqlite3VdbeMemConsistentDualRep(Mem *p){
+SQLITE_PRIVATE int sqlite3VdbeMemValidStrRep(Mem *p){
   char zBuf[100];
   char *z;
   int i, j, incr;
   if( (p->flags & MEM_Str)==0 ) return 1;
-  if( (p->flags & (MEM_Int|MEM_Real))==0 ) return 1;
-  if( p->flags & MEM_Int ){
-    sqlite3_snprintf(sizeof(zBuf),zBuf,"%lld",p->u.i);
-  }else{
-    sqlite3_snprintf(sizeof(zBuf),zBuf,"%!.15g",p->u.r);
-  }
+  if( p->flags & MEM_Term ){
+    /* Insure that the string is properly zero-terminated.  Pay particular
+    ** attention to the case where p->n is odd */
+    if( p->szMalloc>0 && p->z==p->zMalloc ){
+      assert( p->enc==SQLITE_UTF8 || p->szMalloc >= ((p->n+1)&~1)+2 );
+      assert( p->enc!=SQLITE_UTF8 || p->szMalloc >= p->n+1 );
+    }
+    assert( p->z[p->n]==0 );
+    assert( p->enc==SQLITE_UTF8 || p->z[(p->n+1)&~1]==0 );
+    assert( p->enc==SQLITE_UTF8 || p->z[((p->n+1)&~1)+1]==0 );
+  }
+  if( (p->flags & (MEM_Int|MEM_Real|MEM_IntReal))==0 ) return 1;
+  vdbeMemRenderNum(sizeof(zBuf), zBuf, p);
   z = p->z;
   i = j = 0;
   incr = 1;
@@ -74619,7 +75659,13 @@ SQLITE_PRIVATE SQLITE_NOINLINE int sqlite3VdbeMemGrow(Mem *pMem, int n, int bPre
   assert( pMem->szMalloc==0
        || pMem->szMalloc==sqlite3DbMallocSize(pMem->db, pMem->zMalloc) );
   if( pMem->szMalloc>0 && bPreserve && pMem->z==pMem->zMalloc ){
-    pMem->z = pMem->zMalloc = sqlite3DbReallocOrFree(pMem->db, pMem->z, n);
+    if( pMem->db ){
+      pMem->z = pMem->zMalloc = sqlite3DbReallocOrFree(pMem->db, pMem->z, n);
+    }else{
+      pMem->zMalloc = sqlite3Realloc(pMem->z, n);
+      if( pMem->zMalloc==0 ) sqlite3_free(pMem->z);
+      pMem->z = pMem->zMalloc;
+    }
     bPreserve = 0;
   }else{
     if( pMem->szMalloc>0 ) sqlite3DbFreeNN(pMem->db, pMem->zMalloc);
@@ -74655,8 +75701,8 @@ SQLITE_PRIVATE SQLITE_NOINLINE int sqlite3VdbeMemGrow(Mem *pMem, int n, int bPre
 **
 ** Any prior string or blob content in the pMem object may be discarded.
 ** The pMem->xDel destructor is called, if it exists.  Though MEM_Str
-** and MEM_Blob values may be discarded, MEM_Int, MEM_Real, and MEM_Null
-** values are preserved.
+** and MEM_Blob values may be discarded, MEM_Int, MEM_Real, MEM_IntReal,
+** and MEM_Null values are preserved.
 **
 ** Return SQLITE_OK on success or an error code (probably SQLITE_NOMEM)
 ** if unable to complete the resizing.
@@ -74669,20 +75715,26 @@ SQLITE_PRIVATE int sqlite3VdbeMemClearAndResize(Mem *pMem, int szNew){
   }
   assert( (pMem->flags & MEM_Dyn)==0 );
   pMem->z = pMem->zMalloc;
-  pMem->flags &= (MEM_Null|MEM_Int|MEM_Real);
+  pMem->flags &= (MEM_Null|MEM_Int|MEM_Real|MEM_IntReal);
   return SQLITE_OK;
 }
 
 /*
 ** It is already known that pMem contains an unterminated string.
 ** Add the zero terminator.
+**
+** Three bytes of zero are added.  In this way, there is guaranteed
+** to be a double-zero byte at an even byte boundary in order to
+** terminate a UTF16 string, even if the initial size of the buffer
+** is an odd number of bytes.
 */
 static SQLITE_NOINLINE int vdbeMemAddTerminator(Mem *pMem){
-  if( sqlite3VdbeMemGrow(pMem, pMem->n+2, 1) ){
+  if( sqlite3VdbeMemGrow(pMem, pMem->n+3, 1) ){
     return SQLITE_NOMEM_BKPT;
   }
   pMem->z[pMem->n] = 0;
   pMem->z[pMem->n+1] = 0;
+  pMem->z[pMem->n+2] = 0;
   pMem->flags |= MEM_Term;
   return SQLITE_OK;
 }
@@ -74756,12 +75808,12 @@ SQLITE_PRIVATE int sqlite3VdbeMemNulTerminate(Mem *pMem){
 }
 
 /*
-** Add MEM_Str to the set of representations for the given Mem.  Numbers
-** are converted using sqlite3_snprintf().  Converting a BLOB to a string
-** is a no-op.
+** Add MEM_Str to the set of representations for the given Mem.  This
+** routine is only called if pMem is a number of some kind, not a NULL
+** or a BLOB.
 **
-** Existing representations MEM_Int and MEM_Real are invalidated if
-** bForce is true but are retained if bForce is false.
+** Existing representations MEM_Int, MEM_Real, or MEM_IntReal are invalidated
+** if bForce is true but are retained if bForce is false.
 **
 ** A MEM_Null value will never be passed to this function. This function is
 ** used for converting values to text for returning to the user (i.e. via
@@ -74770,13 +75822,12 @@ SQLITE_PRIVATE int sqlite3VdbeMemNulTerminate(Mem *pMem){
 ** user and the latter is an internal programming error.
 */
 SQLITE_PRIVATE int sqlite3VdbeMemStringify(Mem *pMem, u8 enc, u8 bForce){
-  int fg = pMem->flags;
   const int nByte = 32;
 
   assert( pMem->db==0 || sqlite3_mutex_held(pMem->db->mutex) );
-  assert( !(fg&MEM_Zero) );
-  assert( !(fg&(MEM_Str|MEM_Blob)) );
-  assert( fg&(MEM_Int|MEM_Real) );
+  assert( !(pMem->flags&MEM_Zero) );
+  assert( !(pMem->flags&(MEM_Str|MEM_Blob)) );
+  assert( pMem->flags&(MEM_Int|MEM_Real|MEM_IntReal) );
   assert( !sqlite3VdbeMemIsRowSet(pMem) );
   assert( EIGHT_BYTE_ALIGNMENT(pMem) );
 
@@ -74786,23 +75837,12 @@ SQLITE_PRIVATE int sqlite3VdbeMemStringify(Mem *pMem, u8 enc, u8 bForce){
     return SQLITE_NOMEM_BKPT;
   }
 
-  /* For a Real or Integer, use sqlite3_snprintf() to produce the UTF-8
-  ** string representation of the value. Then, if the required encoding
-  ** is UTF-16le or UTF-16be do a translation.
-  ** 
-  ** FIX ME: It would be better if sqlite3_snprintf() could do UTF-16.
-  */
-  if( fg & MEM_Int ){
-    sqlite3_snprintf(nByte, pMem->z, "%lld", pMem->u.i);
-  }else{
-    assert( fg & MEM_Real );
-    sqlite3_snprintf(nByte, pMem->z, "%!.15g", pMem->u.r);
-  }
+  vdbeMemRenderNum(nByte, pMem->z, pMem);
   assert( pMem->z!=0 );
   pMem->n = sqlite3Strlen30NN(pMem->z);
   pMem->enc = SQLITE_UTF8;
   pMem->flags |= MEM_Str|MEM_Term;
-  if( bForce ) pMem->flags &= ~(MEM_Int|MEM_Real);
+  if( bForce ) pMem->flags &= ~(MEM_Int|MEM_Real|MEM_IntReal);
   sqlite3VdbeChangeEncoding(pMem, enc);
   return SQLITE_OK;
 }
@@ -74847,15 +75887,11 @@ SQLITE_PRIVATE int sqlite3VdbeMemFinalize(Mem *pMem, FuncDef *pFunc){
 #ifndef SQLITE_OMIT_WINDOWFUNC
 SQLITE_PRIVATE int sqlite3VdbeMemAggValue(Mem *pAccum, Mem *pOut, FuncDef *pFunc){
   sqlite3_context ctx;
-  Mem t;
   assert( pFunc!=0 );
   assert( pFunc->xValue!=0 );
   assert( (pAccum->flags & MEM_Null)!=0 || pFunc==pAccum->u.pDef );
   assert( pAccum->db==0 || sqlite3_mutex_held(pAccum->db->mutex) );
   memset(&ctx, 0, sizeof(ctx));
-  memset(&t, 0, sizeof(t));
-  t.flags = MEM_Null;
-  t.db = pAccum->db;
   sqlite3VdbeMemSetNull(pOut);
   ctx.pOut = pOut;
   ctx.pMem = pAccum;
@@ -74976,12 +76012,12 @@ SQLITE_PRIVATE i64 sqlite3VdbeIntValue(Mem *pMem){
   assert( pMem->db==0 || sqlite3_mutex_held(pMem->db->mutex) );
   assert( EIGHT_BYTE_ALIGNMENT(pMem) );
   flags = pMem->flags;
-  if( flags & MEM_Int ){
+  if( flags & (MEM_Int|MEM_IntReal) ){
+    testcase( flags & MEM_IntReal );
     return pMem->u.i;
   }else if( flags & MEM_Real ){
     return doubleToInt64(pMem->u.r);
-  }else if( flags & (MEM_Str|MEM_Blob) ){
-    assert( pMem->z || pMem->n==0 );
+  }else if( (flags & (MEM_Str|MEM_Blob))!=0 && pMem->z!=0 ){
     return memIntValue(pMem);
   }else{
     return 0;
@@ -75005,7 +76041,8 @@ SQLITE_PRIVATE double sqlite3VdbeRealValue(Mem *pMem){
   assert( EIGHT_BYTE_ALIGNMENT(pMem) );
   if( pMem->flags & MEM_Real ){
     return pMem->u.r;
-  }else if( pMem->flags & MEM_Int ){
+  }else if( pMem->flags & (MEM_Int|MEM_IntReal) ){
+    testcase( pMem->flags & MEM_IntReal );
     return (double)pMem->u.i;
   }else if( pMem->flags & (MEM_Str|MEM_Blob) ){
     return memRealValue(pMem);
@@ -75020,7 +76057,8 @@ SQLITE_PRIVATE double sqlite3VdbeRealValue(Mem *pMem){
 ** Return the value ifNull if pMem is NULL.  
 */
 SQLITE_PRIVATE int sqlite3VdbeBooleanValue(Mem *pMem, int ifNull){
-  if( pMem->flags & MEM_Int ) return pMem->u.i!=0;
+  testcase( pMem->flags & MEM_IntReal );
+  if( pMem->flags & (MEM_Int|MEM_IntReal) ) return pMem->u.i!=0;
   if( pMem->flags & MEM_Null ) return ifNull;
   return sqlite3VdbeRealValue(pMem)!=0.0;
 }
@@ -75083,17 +76121,21 @@ SQLITE_PRIVATE int sqlite3VdbeMemRealify(Mem *pMem){
 /* Compare a floating point value to an integer.  Return true if the two
 ** values are the same within the precision of the floating point value.
 **
+** This function assumes that i was obtained by assignment from r1.
+**
 ** For some versions of GCC on 32-bit machines, if you do the more obvious
 ** comparison of "r1==(double)i" you sometimes get an answer of false even
 ** though the r1 and (double)i values are bit-for-bit the same.
 */
-static int sqlite3RealSameAsInt(double r1, sqlite3_int64 i){
+SQLITE_PRIVATE int sqlite3RealSameAsInt(double r1, sqlite3_int64 i){
   double r2 = (double)i;
-  return memcmp(&r1, &r2, sizeof(r1))==0;
+  return r1==0.0
+      || (memcmp(&r1, &r2, sizeof(r1))==0
+          && i >= -2251799813685248LL && i < 2251799813685248LL);
 }
 
 /*
-** Convert pMem so that it has types MEM_Real or MEM_Int or both.
+** Convert pMem so that it has type MEM_Real or MEM_Int.
 ** Invalidate any prior representations.
 **
 ** Every effort is made to force the conversion, even if the input
@@ -75101,25 +76143,26 @@ static int sqlite3RealSameAsInt(double r1, sqlite3_int64 i){
 ** as much of the string as we can and ignore the rest.
 */
 SQLITE_PRIVATE int sqlite3VdbeMemNumerify(Mem *pMem){
-  if( (pMem->flags & (MEM_Int|MEM_Real|MEM_Null))==0 ){
+  testcase( pMem->flags & MEM_Int );
+  testcase( pMem->flags & MEM_Real );
+  testcase( pMem->flags & MEM_IntReal );
+  testcase( pMem->flags & MEM_Null );
+  if( (pMem->flags & (MEM_Int|MEM_Real|MEM_IntReal|MEM_Null))==0 ){
     int rc;
+    sqlite3_int64 ix;
     assert( (pMem->flags & (MEM_Blob|MEM_Str))!=0 );
     assert( pMem->db==0 || sqlite3_mutex_held(pMem->db->mutex) );
-    rc = sqlite3Atoi64(pMem->z, &pMem->u.i, pMem->n, pMem->enc);
-    if( rc==0 ){
+    rc = sqlite3AtoF(pMem->z, &pMem->u.r, pMem->n, pMem->enc);
+    if( ((rc==0 || rc==1) && sqlite3Atoi64(pMem->z, &ix, pMem->n, pMem->enc)<=1)
+     || sqlite3RealSameAsInt(pMem->u.r, (ix = (i64)pMem->u.r))
+    ){
+      pMem->u.i = ix;
       MemSetTypeFlag(pMem, MEM_Int);
     }else{
-      i64 i = pMem->u.i;
-      sqlite3AtoF(pMem->z, &pMem->u.r, pMem->n, pMem->enc);
-      if( rc==1 && sqlite3RealSameAsInt(pMem->u.r, i) ){
-        pMem->u.i = i;
-        MemSetTypeFlag(pMem, MEM_Int);
-      }else{
-        MemSetTypeFlag(pMem, MEM_Real);
-      }
+      MemSetTypeFlag(pMem, MEM_Real);
     }
   }
-  assert( (pMem->flags & (MEM_Int|MEM_Real|MEM_Null))!=0 );
+  assert( (pMem->flags & (MEM_Int|MEM_Real|MEM_IntReal|MEM_Null))!=0 );
   pMem->flags &= ~(MEM_Str|MEM_Blob|MEM_Zero);
   return SQLITE_OK;
 }
@@ -75131,8 +76174,8 @@ SQLITE_PRIVATE int sqlite3VdbeMemNumerify(Mem *pMem){
 ** affinity even if that results in loss of data.  This routine is
 ** used (for example) to implement the SQL "cast()" operator.
 */
-SQLITE_PRIVATE void sqlite3VdbeMemCast(Mem *pMem, u8 aff, u8 encoding){
-  if( pMem->flags & MEM_Null ) return;
+SQLITE_PRIVATE int sqlite3VdbeMemCast(Mem *pMem, u8 aff, u8 encoding){
+  if( pMem->flags & MEM_Null ) return SQLITE_OK;
   switch( aff ){
     case SQLITE_AFF_BLOB: {   /* Really a cast to BLOB */
       if( (pMem->flags & MEM_Blob)==0 ){
@@ -75162,10 +76205,11 @@ SQLITE_PRIVATE void sqlite3VdbeMemCast(Mem *pMem, u8 aff, u8 encoding){
       pMem->flags |= (pMem->flags&MEM_Blob)>>3;
       sqlite3ValueApplyAffinity(pMem, SQLITE_AFF_TEXT, encoding);
       assert( pMem->flags & MEM_Str || pMem->db->mallocFailed );
-      pMem->flags &= ~(MEM_Int|MEM_Real|MEM_Blob|MEM_Zero);
-      break;
+      pMem->flags &= ~(MEM_Int|MEM_Real|MEM_IntReal|MEM_Blob|MEM_Zero);
+      return sqlite3VdbeChangeEncoding(pMem, encoding);
     }
   }
+  return SQLITE_OK;
 }
 
 /*
@@ -75331,23 +76375,30 @@ SQLITE_PRIVATE int sqlite3VdbeMemTooBig(Mem *p){
 ** its link to a shallow copy and by marking any current shallow
 ** copies of this cell as invalid.
 **
-** This is used for testing and debugging only - to make sure shallow
-** copies are not misused.
+** This is used for testing and debugging only - to help ensure that shallow
+** copies (created by OP_SCopy) are not misused.
 */
 SQLITE_PRIVATE void sqlite3VdbeMemAboutToChange(Vdbe *pVdbe, Mem *pMem){
   int i;
   Mem *pX;
-  for(i=0, pX=pVdbe->aMem; i<pVdbe->nMem; i++, pX++){
+  for(i=1, pX=pVdbe->aMem+1; i<pVdbe->nMem; i++, pX++){
     if( pX->pScopyFrom==pMem ){
+      u16 mFlags;
+      if( pVdbe->db->flags & SQLITE_VdbeTrace ){
+        sqlite3DebugPrintf("Invalidate R[%d] due to change in R[%d]\n",
+          (int)(pX - pVdbe->aMem), (int)(pMem - pVdbe->aMem));
+      }
       /* If pX is marked as a shallow copy of pMem, then verify that
       ** no significant changes have been made to pX since the OP_SCopy.
       ** A significant change would indicated a missed call to this
       ** function for pX.  Minor changes, such as adding or removing a
       ** dual type, are allowed, as long as the underlying value is the
       ** same. */
-      u16 mFlags = pMem->flags & pX->flags & pX->mScopyFlags;
-      assert( (mFlags&MEM_Int)==0 || pMem->u.i==pX->u.i );
-      assert( (mFlags&MEM_Real)==0 || pMem->u.r==pX->u.r );
+      mFlags = pMem->flags & pX->flags & pX->mScopyFlags;
+      assert( (mFlags&(MEM_Int|MEM_IntReal))==0 || pMem->u.i==pX->u.i );
+      /* assert( (mFlags&MEM_Real)==0 || pMem->u.r==pX->u.r ); */
+      /*                                          ^^           */
+      /*       Cannot reliably compare doubles for equality    */
       assert( (mFlags&MEM_Str)==0  || (pMem->n==pX->n && pMem->z==pX->z) );
       assert( (mFlags&MEM_Blob)==0  || sqlite3BlobCompare(pMem,pX)==0 );
       
@@ -75361,7 +76412,6 @@ SQLITE_PRIVATE void sqlite3VdbeMemAboutToChange(Vdbe *pVdbe, Mem *pMem){
 }
 #endif /* SQLITE_DEBUG */
 
-
 /*
 ** Make an shallow copy of pFrom into pTo.  Prior contents of
 ** pTo are freed.  The pFrom->z field is not duplicated.  If
@@ -75507,10 +76557,19 @@ SQLITE_PRIVATE int sqlite3VdbeMemSetStr(
 
   pMem->n = nByte;
   pMem->flags = flags;
-  pMem->enc = (enc==0 ? SQLITE_UTF8 : enc);
+  if( enc ){
+    pMem->enc = enc;
+#ifdef SQLITE_ENABLE_SESSION
+  }else if( pMem->db==0 ){
+    pMem->enc = SQLITE_UTF8;
+#endif
+  }else{
+    assert( pMem->db!=0 );
+    pMem->enc = ENC(pMem->db);
+  }
 
 #ifndef SQLITE_OMIT_UTF16
-  if( pMem->enc!=SQLITE_UTF8 && sqlite3VdbeMemHandleBom(pMem) ){
+  if( enc>SQLITE_UTF8 && sqlite3VdbeMemHandleBom(pMem) ){
     return SQLITE_NOMEM_BKPT;
   }
 #endif
@@ -75621,7 +76680,7 @@ static SQLITE_NOINLINE const void *valueToText(sqlite3_value* pVal, u8 enc){
   assert(pVal->enc==(enc & ~SQLITE_UTF16_ALIGNED) || pVal->db==0
               || pVal->db->mallocFailed );
   if( pVal->enc==(enc & ~SQLITE_UTF16_ALIGNED) ){
-    assert( sqlite3VdbeMemConsistentDualRep(pVal) );
+    assert( sqlite3VdbeMemValidStrRep(pVal) );
     return pVal->z;
   }else{
     return 0;
@@ -75644,7 +76703,7 @@ SQLITE_PRIVATE const void *sqlite3ValueText(sqlite3_value* pVal, u8 enc){
   assert( (enc&3)==(enc&~SQLITE_UTF16_ALIGNED) );
   assert( !sqlite3VdbeMemIsRowSet(pVal) );
   if( (pVal->flags&(MEM_Str|MEM_Term))==(MEM_Str|MEM_Term) && pVal->enc==enc ){
-    assert( sqlite3VdbeMemConsistentDualRep(pVal) );
+    assert( sqlite3VdbeMemValidStrRep(pVal) );
     return pVal->z;
   }
   if( pVal->flags&MEM_Null ){
@@ -75688,7 +76747,7 @@ struct ValueNewStat4Ctx {
 ** an sqlite3_value within the UnpackedRecord.a[] array.
 */
 static sqlite3_value *valueNew(sqlite3 *db, struct ValueNewStat4Ctx *p){
-#ifdef SQLITE_ENABLE_STAT3_OR_STAT4
+#ifdef SQLITE_ENABLE_STAT4
   if( p ){
     UnpackedRecord *pRec = p->ppRec[0];
 
@@ -75724,7 +76783,7 @@ static sqlite3_value *valueNew(sqlite3 *db, struct ValueNewStat4Ctx *p){
   }
 #else
   UNUSED_PARAMETER(p);
-#endif /* defined(SQLITE_ENABLE_STAT3_OR_STAT4) */
+#endif /* defined(SQLITE_ENABLE_STAT4) */
   return sqlite3ValueNew(db);
 }
 
@@ -75748,7 +76807,7 @@ static sqlite3_value *valueNew(sqlite3 *db, struct ValueNewStat4Ctx *p){
 ** and sets (*ppVal) to NULL. Or, if an error occurs, (*ppVal) is set to
 ** NULL and an SQLite error code returned.
 */
-#ifdef SQLITE_ENABLE_STAT3_OR_STAT4
+#ifdef SQLITE_ENABLE_STAT4
 static int valueFromFunction(
   sqlite3 *db,                    /* The database connection */
   Expr *p,                        /* The expression to evaluate */
@@ -75831,7 +76890,7 @@ static int valueFromFunction(
 }
 #else
 # define valueFromFunction(a,b,c,d,e,f) SQLITE_OK
-#endif /* defined(SQLITE_ENABLE_STAT3_OR_STAT4) */
+#endif /* defined(SQLITE_ENABLE_STAT4) */
 
 /*
 ** Extract a value from the supplied expression in the manner described
@@ -75860,7 +76919,7 @@ static int valueFromExpr(
 
   assert( pExpr!=0 );
   while( (op = pExpr->op)==TK_UPLUS || op==TK_SPAN ) pExpr = pExpr->pLeft;
-#if defined(SQLITE_ENABLE_STAT3_OR_STAT4)
+#if defined(SQLITE_ENABLE_STAT4)
   if( op==TK_REGISTER ) op = pExpr->op2;
 #else
   if( NEVER(op==TK_REGISTER) ) op = pExpr->op2;
@@ -75909,7 +76968,12 @@ static int valueFromExpr(
     }else{
       sqlite3ValueApplyAffinity(pVal, affinity, SQLITE_UTF8);
     }
-    if( pVal->flags & (MEM_Int|MEM_Real) ) pVal->flags &= ~MEM_Str;
+    assert( (pVal->flags & MEM_IntReal)==0 );
+    if( pVal->flags & (MEM_Int|MEM_IntReal|MEM_Real) ){
+      testcase( pVal->flags & MEM_Int );
+      testcase( pVal->flags & MEM_Real );
+      pVal->flags &= ~MEM_Str;
+    }
     if( enc!=SQLITE_UTF8 ){
       rc = sqlite3VdbeChangeEncoding(pVal, enc);
     }
@@ -75922,7 +76986,11 @@ static int valueFromExpr(
       if( pVal->flags & MEM_Real ){
         pVal->u.r = -pVal->u.r;
       }else if( pVal->u.i==SMALLEST_INT64 ){
+#ifndef SQLITE_OMIT_FLOATING_POINT
         pVal->u.r = -(double)SMALLEST_INT64;
+#else
+        pVal->u.r = LARGEST_INT64;
+#endif
         MemSetTypeFlag(pVal, MEM_Real);
       }else{
         pVal->u.i = -pVal->u.i;
@@ -75932,7 +77000,7 @@ static int valueFromExpr(
   }else if( op==TK_NULL ){
     pVal = valueNew(db, pCtx);
     if( pVal==0 ) goto no_mem;
-    sqlite3VdbeMemNumerify(pVal);
+    sqlite3VdbeMemSetNull(pVal);
   }
 #ifndef SQLITE_OMIT_BLOB_LITERAL
   else if( op==TK_BLOB ){
@@ -75948,7 +77016,7 @@ static int valueFromExpr(
                          0, SQLITE_DYNAMIC);
   }
 #endif
-#ifdef SQLITE_ENABLE_STAT3_OR_STAT4
+#ifdef SQLITE_ENABLE_STAT4
   else if( op==TK_FUNCTION && pCtx!=0 ){
     rc = valueFromFunction(db, pExpr, enc, affinity, &pVal, pCtx);
   }
@@ -75965,13 +77033,13 @@ static int valueFromExpr(
   return rc;
 
 no_mem:
-#ifdef SQLITE_ENABLE_STAT3_OR_STAT4
+#ifdef SQLITE_ENABLE_STAT4
   if( pCtx==0 || pCtx->pParse->nErr==0 )
 #endif
     sqlite3OomFault(db);
   sqlite3DbFree(db, zVal);
   assert( *ppVal==0 );
-#ifdef SQLITE_ENABLE_STAT3_OR_STAT4
+#ifdef SQLITE_ENABLE_STAT4
   if( pCtx==0 ) sqlite3ValueFree(pVal);
 #else
   assert( pCtx==0 ); sqlite3ValueFree(pVal);
@@ -75999,56 +77067,7 @@ SQLITE_PRIVATE int sqlite3ValueFromExpr(
   return pExpr ? valueFromExpr(db, pExpr, enc, affinity, ppVal, 0) : 0;
 }
 
-#ifdef SQLITE_ENABLE_STAT3_OR_STAT4
-/*
-** The implementation of the sqlite_record() function. This function accepts
-** a single argument of any type. The return value is a formatted database 
-** record (a blob) containing the argument value.
-**
-** This is used to convert the value stored in the 'sample' column of the
-** sqlite_stat3 table to the record format SQLite uses internally.
-*/
-static void recordFunc(
-  sqlite3_context *context,
-  int argc,
-  sqlite3_value **argv
-){
-  const int file_format = 1;
-  u32 iSerial;                    /* Serial type */
-  int nSerial;                    /* Bytes of space for iSerial as varint */
-  u32 nVal;                       /* Bytes of space required for argv[0] */
-  int nRet;
-  sqlite3 *db;
-  u8 *aRet;
-
-  UNUSED_PARAMETER( argc );
-  iSerial = sqlite3VdbeSerialType(argv[0], file_format, &nVal);
-  nSerial = sqlite3VarintLen(iSerial);
-  db = sqlite3_context_db_handle(context);
-
-  nRet = 1 + nSerial + nVal;
-  aRet = sqlite3DbMallocRawNN(db, nRet);
-  if( aRet==0 ){
-    sqlite3_result_error_nomem(context);
-  }else{
-    aRet[0] = nSerial+1;
-    putVarint32(&aRet[1], iSerial);
-    sqlite3VdbeSerialPut(&aRet[1+nSerial], argv[0], iSerial);
-    sqlite3_result_blob(context, aRet, nRet, SQLITE_TRANSIENT);
-    sqlite3DbFreeNN(db, aRet);
-  }
-}
-
-/*
-** Register built-in functions used to help read ANALYZE data.
-*/
-SQLITE_PRIVATE void sqlite3AnalyzeFunctions(void){
-  static FuncDef aAnalyzeTableFuncs[] = {
-    FUNCTION(sqlite_record,   1, 0, 0, recordFunc),
-  };
-  sqlite3InsertBuiltinFuncs(aAnalyzeTableFuncs, ArraySize(aAnalyzeTableFuncs));
-}
-
+#ifdef SQLITE_ENABLE_STAT4
 /*
 ** Attempt to extract a value from pExpr and use it to construct *ppVal.
 **
@@ -76330,6 +77349,10 @@ SQLITE_PRIVATE int sqlite3ValueBytes(sqlite3_value *pVal, u8 enc){
 /* #include "sqliteInt.h" */
 /* #include "vdbeInt.h" */
 
+/* Forward references */
+static void freeEphemeralFunction(sqlite3 *db, FuncDef *pDef);
+static void vdbeFreeOpArray(sqlite3 *, Op *, int);
+
 /*
 ** Create a new virtual database engine.
 */
@@ -76357,6 +77380,13 @@ SQLITE_PRIVATE Vdbe *sqlite3VdbeCreate(Parse *pParse){
   return p;
 }
 
+/*
+** Return the Parse object that owns a Vdbe object.
+*/
+SQLITE_PRIVATE Parse *sqlite3VdbeParser(Vdbe *p){
+  return p->pParse;
+}
+
 /*
 ** Change the error string stored in Vdbe.zErrMsg
 */
@@ -76437,7 +77467,7 @@ SQLITE_PRIVATE void sqlite3VdbeSwap(Vdbe *pA, Vdbe *pB){
   zTmp = pA->zSql;
   pA->zSql = pB->zSql;
   pB->zSql = zTmp;
-#if 0
+#ifdef SQLITE_ENABLE_NORMALIZE
   zTmp = pA->zNormSql;
   pA->zNormSql = pB->zNormSql;
   pB->zNormSql = zTmp;
@@ -76498,9 +77528,16 @@ static int growOpArray(Vdbe *v, int nOp){
 #ifdef SQLITE_DEBUG
 /* This routine is just a convenient place to set a breakpoint that will
 ** fire after each opcode is inserted and displayed using
-** "PRAGMA vdbe_addoptrace=on".
+** "PRAGMA vdbe_addoptrace=on".  Parameters "pc" (program counter) and
+** pOp are available to make the breakpoint conditional.
+**
+** Other useful labels for breakpoints include:
+**   test_trace_breakpoint(pc,pOp)
+**   sqlite3CorruptError(lineno)
+**   sqlite3MisuseError(lineno)
+**   sqlite3CantopenError(lineno)
 */
-static void test_addop_breakpoint(void){
+static void test_addop_breakpoint(int pc, Op *pOp){
   static int n = 0;
   n++;
 }
@@ -76553,7 +77590,7 @@ SQLITE_PRIVATE int sqlite3VdbeAddOp3(Vdbe *p, int op, int p1, int p2, int p3){
 #ifdef SQLITE_DEBUG
   if( p->db->flags & SQLITE_VdbeAddopTrace ){
     sqlite3VdbePrintOp(0, i, &p->aOp[i]);
-    test_addop_breakpoint();
+    test_addop_breakpoint(i, &p->aOp[i]);
   }
 #endif
 #ifdef VDBE_PROFILE
@@ -76636,6 +77673,49 @@ SQLITE_PRIVATE int sqlite3VdbeAddOp4(
   return addr;
 }
 
+/*
+** Add an OP_Function or OP_PureFunc opcode.
+**
+** The eCallCtx argument is information (typically taken from Expr.op2)
+** that describes the calling context of the function.  0 means a general
+** function call.  NC_IsCheck means called by a check constraint,
+** NC_IdxExpr means called as part of an index expression.  NC_PartIdx
+** means in the WHERE clause of a partial index.  NC_GenCol means called
+** while computing a generated column value.  0 is the usual case.
+*/
+SQLITE_PRIVATE int sqlite3VdbeAddFunctionCall(
+  Parse *pParse,        /* Parsing context */
+  int p1,               /* Constant argument mask */
+  int p2,               /* First argument register */
+  int p3,               /* Register into which results are written */
+  int nArg,             /* Number of argument */
+  const FuncDef *pFunc, /* The function to be invoked */
+  int eCallCtx          /* Calling context */
+){
+  Vdbe *v = pParse->pVdbe;
+  int nByte;
+  int addr;
+  sqlite3_context *pCtx;
+  assert( v );
+  nByte = sizeof(*pCtx) + (nArg-1)*sizeof(sqlite3_value*);
+  pCtx = sqlite3DbMallocRawNN(pParse->db, nByte);
+  if( pCtx==0 ){
+    assert( pParse->db->mallocFailed );
+    freeEphemeralFunction(pParse->db, (FuncDef*)pFunc);
+    return 0;
+  }
+  pCtx->pOut = 0;
+  pCtx->pFunc = (FuncDef*)pFunc;
+  pCtx->pVdbe = 0;
+  pCtx->isError = 0;
+  pCtx->argc = nArg;
+  pCtx->iOp = sqlite3VdbeCurrentAddr(v);
+  addr = sqlite3VdbeAddOp4(v, eCallCtx ? OP_PureFunc : OP_Function,
+                           p1, p2, p3, (char*)pCtx, P4_FUNCCTX);
+  sqlite3VdbeChangeP5(v, eCallCtx & NC_SelfRef);
+  return addr;
+}
+
 /*
 ** Add an opcode that includes the p4 value with a P4_INT64 or
 ** P4_REAL type.
@@ -76928,6 +78008,7 @@ static Op *opIterNext(VdbeOpIter *p){
 **   *  OP_HaltIfNull with P1=SQLITE_CONSTRAINT and P2=OE_Abort.
 **   *  OP_Destroy
 **   *  OP_VUpdate
+**   *  OP_VCreate
 **   *  OP_VRename
 **   *  OP_FkCounter with P2==0 (immediate foreign key constraint)
 **   *  OP_CreateBtree/BTREE_INTKEY and OP_InitCoroutine 
@@ -76944,6 +78025,7 @@ SQLITE_PRIVATE int sqlite3VdbeAssertMayAbort(Vdbe *v, int mayAbort){
   int hasAbort = 0;
   int hasFkCounter = 0;
   int hasCreateTable = 0;
+  int hasCreateIndex = 0;
   int hasInitCoroutine = 0;
   Op *pOp;
   VdbeOpIter sIter;
@@ -76954,7 +78036,8 @@ SQLITE_PRIVATE int sqlite3VdbeAssertMayAbort(Vdbe *v, int mayAbort){
     int opcode = pOp->opcode;
     if( opcode==OP_Destroy || opcode==OP_VUpdate || opcode==OP_VRename 
      || opcode==OP_VDestroy
-     || (opcode==OP_Function0 && pOp->p4.pFunc->funcFlags&SQLITE_FUNC_INTERNAL)
+     || opcode==OP_VCreate
+     || (opcode==OP_ParseSchema && pOp->p4.z==0)
      || ((opcode==OP_Halt || opcode==OP_HaltIfNull) 
       && ((pOp->p1)!=SQLITE_OK && pOp->p2==OE_Abort))
     ){
@@ -76962,6 +78045,14 @@ SQLITE_PRIVATE int sqlite3VdbeAssertMayAbort(Vdbe *v, int mayAbort){
       break;
     }
     if( opcode==OP_CreateBtree && pOp->p3==BTREE_INTKEY ) hasCreateTable = 1;
+    if( mayAbort ){
+      /* hasCreateIndex may also be set for some DELETE statements that use
+      ** OP_Clear. So this routine may end up returning true in the case 
+      ** where a "DELETE FROM tbl" has a statement-journal but does not
+      ** require one. This is not so bad - it is an inefficiency, not a bug. */
+      if( opcode==OP_CreateBtree && pOp->p3==BTREE_BLOBKEY ) hasCreateIndex = 1;
+      if( opcode==OP_Clear ) hasCreateIndex = 1;
+    }
     if( opcode==OP_InitCoroutine ) hasInitCoroutine = 1;
 #ifndef SQLITE_OMIT_FOREIGN_KEY
     if( opcode==OP_FkCounter && pOp->p1==0 && pOp->p2==1 ){
@@ -76977,7 +78068,8 @@ SQLITE_PRIVATE int sqlite3VdbeAssertMayAbort(Vdbe *v, int mayAbort){
   ** true for this case to prevent the assert() in the callers frame
   ** from failing.  */
   return ( v->db->mallocFailed || hasAbort==mayAbort || hasFkCounter
-              || (hasCreateTable && hasInitCoroutine) );
+        || (hasCreateTable && hasInitCoroutine) || hasCreateIndex
+  );
 }
 #endif /* SQLITE_DEBUG - the sqlite3AssertMayAbort() function */
 
@@ -77282,16 +78374,16 @@ SQLITE_PRIVATE void sqlite3VdbeScanStatus(
 ** Change the value of the opcode, or P1, P2, P3, or P5 operands
 ** for a specific instruction.
 */
-SQLITE_PRIVATE void sqlite3VdbeChangeOpcode(Vdbe *p, u32 addr, u8 iNewOpcode){
+SQLITE_PRIVATE void sqlite3VdbeChangeOpcode(Vdbe *p, int addr, u8 iNewOpcode){
   sqlite3VdbeGetOp(p,addr)->opcode = iNewOpcode;
 }
-SQLITE_PRIVATE void sqlite3VdbeChangeP1(Vdbe *p, u32 addr, int val){
+SQLITE_PRIVATE void sqlite3VdbeChangeP1(Vdbe *p, int addr, int val){
   sqlite3VdbeGetOp(p,addr)->p1 = val;
 }
-SQLITE_PRIVATE void sqlite3VdbeChangeP2(Vdbe *p, u32 addr, int val){
+SQLITE_PRIVATE void sqlite3VdbeChangeP2(Vdbe *p, int addr, int val){
   sqlite3VdbeGetOp(p,addr)->p2 = val;
 }
-SQLITE_PRIVATE void sqlite3VdbeChangeP3(Vdbe *p, u32 addr, int val){
+SQLITE_PRIVATE void sqlite3VdbeChangeP3(Vdbe *p, int addr, int val){
   sqlite3VdbeGetOp(p,addr)->p3 = val;
 }
 SQLITE_PRIVATE void sqlite3VdbeChangeP5(Vdbe *p, u16 p5){
@@ -77318,8 +78410,6 @@ static void freeEphemeralFunction(sqlite3 *db, FuncDef *pDef){
   }
 }
 
-static void vdbeFreeOpArray(sqlite3 *, Op *, int);
-
 /*
 ** Delete a P4 value if necessary.
 */
@@ -77329,7 +78419,7 @@ static SQLITE_NOINLINE void freeP4Mem(sqlite3 *db, Mem *p){
 }
 static SQLITE_NOINLINE void freeP4FuncCtx(sqlite3 *db, sqlite3_context *p){
   freeEphemeralFunction(db, p->pFunc);
- sqlite3DbFreeNN(db, p);
+  sqlite3DbFreeNN(db, p);
 }
 static void freeP4(sqlite3 *db, int p4type, void *p4){
   assert( db );
@@ -77403,6 +78493,13 @@ SQLITE_PRIVATE void sqlite3VdbeLinkSubProgram(Vdbe *pVdbe, SubProgram *p){
   pVdbe->pProgram = p;
 }
 
+/*
+** Return true if the given Vdbe has any SubPrograms.
+*/
+SQLITE_PRIVATE int sqlite3VdbeHasSubProgram(Vdbe *pVdbe){
+  return pVdbe->pProgram!=0;
+}
+
 /*
 ** Change the opcode at addr into OP_Noop
 */
@@ -77430,6 +78527,41 @@ SQLITE_PRIVATE int sqlite3VdbeDeletePriorOpcode(Vdbe *p, u8 op){
   }
 }
 
+#ifdef SQLITE_DEBUG
+/*
+** Generate an OP_ReleaseReg opcode to indicate that a range of
+** registers, except any identified by mask, are no longer in use.
+*/
+SQLITE_PRIVATE void sqlite3VdbeReleaseRegisters(
+  Parse *pParse,       /* Parsing context */
+  int iFirst,          /* Index of first register to be released */
+  int N,               /* Number of registers to release */
+  u32 mask,            /* Mask of registers to NOT release */
+  int bUndefine        /* If true, mark registers as undefined */
+){
+  if( N==0 ) return;
+  assert( pParse->pVdbe );
+  assert( iFirst>=1 );
+  assert( iFirst+N-1<=pParse->nMem );
+  if( N<=31 && mask!=0 ){
+    while( N>0 && (mask&1)!=0 ){
+      mask >>= 1;
+      iFirst++;
+      N--;
+    }
+    while( N>0 && N<=32 && (mask & MASKBIT32(N-1))!=0 ){
+      mask &= ~MASKBIT32(N-1);
+      N--;
+    }
+  }
+  if( N>0 ){
+    sqlite3VdbeAddOp3(pParse->pVdbe, OP_ReleaseReg, iFirst, N, *(int*)&mask);
+    if( bUndefine ) sqlite3VdbeChangeP5(pParse->pVdbe, 1);
+  }
+}
+#endif /* SQLITE_DEBUG */
+
+
 /*
 ** Change the value of the P4 operand for a specific instruction.
 ** This routine is useful when a large program is loaded from a
@@ -77547,7 +78679,8 @@ SQLITE_PRIVATE void sqlite3VdbeSetP4KeyInfo(Parse *pParse, Index *pIdx){
 */
 static void vdbeVComment(Vdbe *p, const char *zFormat, va_list ap){
   assert( p->nOp>0 || p->aOp==0 );
-  assert( p->aOp==0 || p->aOp[p->nOp-1].zComment==0 || p->db->mallocFailed );
+  assert( p->aOp==0 || p->aOp[p->nOp-1].zComment==0 || p->db->mallocFailed
+          || p->pParse->nErr>0 );
   if( p->nOp ){
     assert( p->aOp );
     sqlite3DbFree(p->db, p->aOp[p->nOp-1].zComment);
@@ -77798,14 +78931,16 @@ static char *displayP4(Op *pOp, char *zTemp, int nTemp){
     case P4_KEYINFO: {
       int j;
       KeyInfo *pKeyInfo = pOp->p4.pKeyInfo;
-      assert( pKeyInfo->aSortOrder!=0 );
+      assert( pKeyInfo->aSortFlags!=0 );
       sqlite3_str_appendf(&x, "k(%d", pKeyInfo->nKeyField);
       for(j=0; j<pKeyInfo->nKeyField; j++){
         CollSeq *pColl = pKeyInfo->aColl[j];
         const char *zColl = pColl ? pColl->zName : "";
         if( strcmp(zColl, "BINARY")==0 ) zColl = "B";
-        sqlite3_str_appendf(&x, ",%s%s", 
-               pKeyInfo->aSortOrder[j] ? "-" : "", zColl);
+        sqlite3_str_appendf(&x, ",%s%s%s", 
+               (pKeyInfo->aSortFlags[j] & KEYINFO_ORDER_DESC) ? "-" : "", 
+               (pKeyInfo->aSortFlags[j] & KEYINFO_ORDER_BIGNULL)? "N." : "", 
+               zColl);
       }
       sqlite3_str_append(&x, ")", 1);
       break;
@@ -77826,13 +78961,11 @@ static char *displayP4(Op *pOp, char *zTemp, int nTemp){
       sqlite3_str_appendf(&x, "%s(%d)", pDef->zName, pDef->nArg);
       break;
     }
-#if defined(SQLITE_DEBUG) || defined(VDBE_PROFILE)
     case P4_FUNCCTX: {
       FuncDef *pDef = pOp->p4.pCtx->pFunc;
       sqlite3_str_appendf(&x, "%s(%d)", pDef->zName, pDef->nArg);
       break;
     }
-#endif
     case P4_INT64: {
       sqlite3_str_appendf(&x, "%lld", *pOp->p4.pI64);
       break;
@@ -77849,7 +78982,7 @@ static char *displayP4(Op *pOp, char *zTemp, int nTemp){
       Mem *pMem = pOp->p4.pMem;
       if( pMem->flags & MEM_Str ){
         zP4 = pMem->z;
-      }else if( pMem->flags & MEM_Int ){
+      }else if( pMem->flags & (MEM_Int|MEM_IntReal) ){
         sqlite3_str_appendf(&x, "%lld", pMem->u.i);
       }else if( pMem->flags & MEM_Real ){
         sqlite3_str_appendf(&x, "%.16g", pMem->u.r);
@@ -78212,8 +79345,11 @@ SQLITE_PRIVATE int sqlite3VdbeList(
       ** pick up the appropriate opcode. */
       int j;
       i -= p->nOp;
+      assert( apSub!=0 );
+      assert( nSub>0 );
       for(j=0; i>=apSub[j]->nOp; j++){
         i -= apSub[j]->nOp;
+        assert( i<apSub[j]->nOp || j+1<nSub );
       }
       pOp = &apSub[j]->aOp[i];
     }
@@ -78523,8 +79659,26 @@ SQLITE_PRIVATE void sqlite3VdbeMakeReady(
 
   resolveP2Values(p, &nArg);
   p->usesStmtJournal = (u8)(pParse->isMultiWrite && pParse->mayAbort);
-  if( pParse->explain && nMem<10 ){
-    nMem = 10;
+  if( pParse->explain ){
+    static const char * const azColName[] = {
+       "addr", "opcode", "p1", "p2", "p3", "p4", "p5", "comment",
+       "id", "parent", "notused", "detail"
+    };
+    int iFirst, mx, i;
+    if( nMem<10 ) nMem = 10;
+    if( pParse->explain==2 ){
+      sqlite3VdbeSetNumCols(p, 4);
+      iFirst = 8;
+      mx = 12;
+    }else{
+      sqlite3VdbeSetNumCols(p, 8);
+      iFirst = 0;
+      mx = 8;
+    }
+    for(i=iFirst; i<mx; i++){
+      sqlite3VdbeSetColName(p, i-iFirst, COLNAME_NAME,
+                            azColName[i], SQLITE_STATIC);
+    }
   }
   p->expired = 0;
 
@@ -78874,7 +80028,7 @@ static int vdbeCommit(sqlite3 *db, Vdbe *p){
 
     /* Select a master journal file name */
     nMainFile = sqlite3Strlen30(zMainFile);
-    zMaster = sqlite3MPrintf(db, "%s-mjXXXXXX9XXz", zMainFile);
+    zMaster = sqlite3MPrintf(db, "%s-mjXXXXXX9XXz%c%c", zMainFile, 0, 0);
     if( zMaster==0 ) return SQLITE_NOMEM_BKPT;
     do {
       u32 iRandom;
@@ -79211,7 +80365,7 @@ SQLITE_PRIVATE int sqlite3VdbeHalt(Vdbe *p){
     }
 
     /* Check for immediate foreign key violations. */
-    if( p->rc==SQLITE_OK ){
+    if( p->rc==SQLITE_OK || (p->errorAction==OE_Fail && !isSpecialError) ){
       sqlite3VdbeCheckFk(p, 0);
     }
   
@@ -79613,7 +80767,7 @@ SQLITE_PRIVATE void sqlite3VdbeDelete(Vdbe *p){
 ** carried out.  Seek the cursor now.  If an error occurs, return
 ** the appropriate error code.
 */
-static int SQLITE_NOINLINE handleDeferredMoveto(VdbeCursor *p){
+SQLITE_PRIVATE int SQLITE_NOINLINE sqlite3VdbeFinishMoveto(VdbeCursor *p){
   int res, rc;
 #ifdef SQLITE_TEST
   extern int sqlite3_search_count;
@@ -79685,7 +80839,7 @@ SQLITE_PRIVATE int sqlite3VdbeCursorMoveto(VdbeCursor **pp, int *piCol){
       *piCol = iMap - 1;
       return SQLITE_OK;
     }
-    return handleDeferredMoveto(p);
+    return sqlite3VdbeFinishMoveto(p);
   }
   if( sqlite3BtreeCursorHasMoved(p->uc.pCursor) ){
     return handleMovedCursor(p);
@@ -79735,8 +80889,17 @@ SQLITE_PRIVATE int sqlite3VdbeCursorMoveto(VdbeCursor **pp, int *piCol){
 ** of SQLite will not understand those serial types.
 */
 
+#if 0 /* Inlined into the OP_MakeRecord opcode */
 /*
 ** Return the serial-type for the value stored in pMem.
+**
+** This routine might convert a large MEM_IntReal value into MEM_Real.
+**
+** 2019-07-11:  The primary user of this subroutine was the OP_MakeRecord
+** opcode in the byte-code engine.  But by moving this routine in-line, we
+** can omit some redundant tests and make that opcode a lot faster.  So
+** this routine is now only used by the STAT3 logic and STAT3 support has
+** ended.  The code is kept here for historical reference only.
 */
 SQLITE_PRIVATE u32 sqlite3VdbeSerialType(Mem *pMem, int file_format, u32 *pLen){
   int flags = pMem->flags;
@@ -79747,11 +80910,13 @@ SQLITE_PRIVATE u32 sqlite3VdbeSerialType(Mem *pMem, int file_format, u32 *pLen){
     *pLen = 0;
     return 0;
   }
-  if( flags&MEM_Int ){
+  if( flags&(MEM_Int|MEM_IntReal) ){
     /* Figure out whether to use 1, 2, 4, 6 or 8 bytes. */
 #   define MAX_6BYTE ((((i64)0x00008000)<<32)-1)
     i64 i = pMem->u.i;
     u64 u;
+    testcase( flags & MEM_Int );
+    testcase( flags & MEM_IntReal );
     if( i<0 ){
       u = ~i;
     }else{
@@ -79771,6 +80936,15 @@ SQLITE_PRIVATE u32 sqlite3VdbeSerialType(Mem *pMem, int file_format, u32 *pLen){
     if( u<=2147483647 ){ *pLen = 4; return 4; }
     if( u<=MAX_6BYTE ){ *pLen = 6; return 5; }
     *pLen = 8;
+    if( flags&MEM_IntReal ){
+      /* If the value is IntReal and is going to take up 8 bytes to store
+      ** as an integer, then we might as well make it an 8-byte floating
+      ** point value */
+      pMem->u.r = (double)pMem->u.i;
+      pMem->flags &= ~MEM_IntReal;
+      pMem->flags |= MEM_Real;
+      return 7;
+    }
     return 6;
   }
   if( flags&MEM_Real ){
@@ -79786,6 +80960,7 @@ SQLITE_PRIVATE u32 sqlite3VdbeSerialType(Mem *pMem, int file_format, u32 *pLen){
   *pLen = n;
   return ((n*2) + 12 + ((flags&MEM_Str)!=0));
 }
+#endif /* inlined into OP_MakeRecord */
 
 /*
 ** The sizes for serial types less than 128
@@ -79944,7 +81119,7 @@ SQLITE_PRIVATE u32 sqlite3VdbeSerialPut(u8 *buf, Mem *pMem, u32 serial_type){
 ** routine so that in most cases the overhead of moving the stack pointer
 ** is avoided.
 */ 
-static u32 SQLITE_NOINLINE serialGet(
+static u32 serialGet(
   const unsigned char *buf,     /* Buffer to deserialize from */
   u32 serial_type,              /* Serial type to deserialize */
   Mem *pMem                     /* Memory cell to write value into */
@@ -79976,7 +81151,7 @@ static u32 SQLITE_NOINLINE serialGet(
     assert( sizeof(x)==8 && sizeof(pMem->u.r)==8 );
     swapMixedEndianFloat(x);
     memcpy(&pMem->u.r, &x, sizeof(x));
-    pMem->flags = sqlite3IsNaN(pMem->u.r) ? MEM_Null : MEM_Real;
+    pMem->flags = IsNaN(x) ? MEM_Null : MEM_Real;
   }
   return 8;
 }
@@ -80094,7 +81269,7 @@ SQLITE_PRIVATE UnpackedRecord *sqlite3VdbeAllocUnpackedRecord(
   p = (UnpackedRecord *)sqlite3DbMallocRaw(pKeyInfo->db, nByte);
   if( !p ) return 0;
   p->aMem = (Mem*)&((char*)p)[ROUND8(sizeof(UnpackedRecord))];
-  assert( pKeyInfo->aSortOrder!=0 );
+  assert( pKeyInfo->aSortFlags!=0 );
   p->pKeyInfo = pKeyInfo;
   p->nField = pKeyInfo->nKeyField + 1;
   return p;
@@ -80193,7 +81368,7 @@ static int vdbeRecordCompareDebug(
   if( szHdr1>98307 ) return SQLITE_CORRUPT;
   d1 = szHdr1;
   assert( pKeyInfo->nAllField>=pPKey2->nField || CORRUPT_DB );
-  assert( pKeyInfo->aSortOrder!=0 );
+  assert( pKeyInfo->aSortFlags!=0 );
   assert( pKeyInfo->nKeyField>0 );
   assert( idx1<=szHdr1 || CORRUPT_DB );
   do{
@@ -80224,7 +81399,12 @@ static int vdbeRecordCompareDebug(
                            pKeyInfo->nAllField>i ? pKeyInfo->aColl[i] : 0);
     if( rc!=0 ){
       assert( mem1.szMalloc==0 );  /* See comment below */
-      if( pKeyInfo->aSortOrder[i] ){
+      if( (pKeyInfo->aSortFlags[i] & KEYINFO_ORDER_BIGNULL)
+       && ((mem1.flags & MEM_Null) || (pPKey2->aMem[i].flags & MEM_Null)) 
+      ){
+        rc = -rc;
+      }
+      if( pKeyInfo->aSortFlags[i] & KEYINFO_ORDER_DESC ){
         rc = -rc;  /* Invert the result for DESC sort order. */
       }
       goto debugCompareEnd;
@@ -80426,8 +81606,13 @@ SQLITE_PRIVATE int sqlite3MemCompare(const Mem *pMem1, const Mem *pMem2, const C
 
   /* At least one of the two values is a number
   */
-  if( combined_flags&(MEM_Int|MEM_Real) ){
-    if( (f1 & f2 & MEM_Int)!=0 ){
+  if( combined_flags&(MEM_Int|MEM_Real|MEM_IntReal) ){
+    testcase( combined_flags & MEM_Int );
+    testcase( combined_flags & MEM_Real );
+    testcase( combined_flags & MEM_IntReal );
+    if( (f1 & f2 & (MEM_Int|MEM_IntReal))!=0 ){
+      testcase( f1 & f2 & MEM_Int );
+      testcase( f1 & f2 & MEM_IntReal );
       if( pMem1->u.i < pMem2->u.i ) return -1;
       if( pMem1->u.i > pMem2->u.i ) return +1;
       return 0;
@@ -80437,15 +81622,23 @@ SQLITE_PRIVATE int sqlite3MemCompare(const Mem *pMem1, const Mem *pMem2, const C
       if( pMem1->u.r > pMem2->u.r ) return +1;
       return 0;
     }
-    if( (f1&MEM_Int)!=0 ){
+    if( (f1&(MEM_Int|MEM_IntReal))!=0 ){
+      testcase( f1 & MEM_Int );
+      testcase( f1 & MEM_IntReal );
       if( (f2&MEM_Real)!=0 ){
         return sqlite3IntFloatCompare(pMem1->u.i, pMem2->u.r);
+      }else if( (f2&(MEM_Int|MEM_IntReal))!=0 ){
+        if( pMem1->u.i < pMem2->u.i ) return -1;
+        if( pMem1->u.i > pMem2->u.i ) return +1;
+        return 0;
       }else{
         return -1;
       }
     }
     if( (f1&MEM_Real)!=0 ){
-      if( (f2&MEM_Int)!=0 ){
+      if( (f2&(MEM_Int|MEM_IntReal))!=0 ){
+        testcase( f2 & MEM_Int );
+        testcase( f2 & MEM_IntReal );
         return -sqlite3IntFloatCompare(pMem2->u.i, pMem1->u.r);
       }else{
         return -1;
@@ -80587,14 +81780,16 @@ SQLITE_PRIVATE int sqlite3VdbeRecordCompareWithSkip(
   VVA_ONLY( mem1.szMalloc = 0; ) /* Only needed by assert() statements */
   assert( pPKey2->pKeyInfo->nAllField>=pPKey2->nField 
        || CORRUPT_DB );
-  assert( pPKey2->pKeyInfo->aSortOrder!=0 );
+  assert( pPKey2->pKeyInfo->aSortFlags!=0 );
   assert( pPKey2->pKeyInfo->nKeyField>0 );
   assert( idx1<=szHdr1 || CORRUPT_DB );
   do{
     u32 serial_type;
 
     /* RHS is an integer */
-    if( pRhs->flags & MEM_Int ){
+    if( pRhs->flags & (MEM_Int|MEM_IntReal) ){
+      testcase( pRhs->flags & MEM_Int );
+      testcase( pRhs->flags & MEM_IntReal );
       serial_type = aKey1[idx1];
       testcase( serial_type==12 );
       if( serial_type>=10 ){
@@ -80708,8 +81903,14 @@ SQLITE_PRIVATE int sqlite3VdbeRecordCompareWithSkip(
     }
 
     if( rc!=0 ){
-      if( pPKey2->pKeyInfo->aSortOrder[i] ){
-        rc = -rc;
+      int sortFlags = pPKey2->pKeyInfo->aSortFlags[i];
+      if( sortFlags ){
+        if( (sortFlags & KEYINFO_ORDER_BIGNULL)==0
+         || ((sortFlags & KEYINFO_ORDER_DESC)
+           !=(serial_type==0 || (pRhs->flags&MEM_Null)))
+        ){
+          rc = -rc;
+        }
       }
       assert( vdbeRecordCompareDebug(nKey1, pKey1, pPKey2, rc) );
       assert( mem1.szMalloc==0 );  /* See comment below */
@@ -80877,7 +82078,11 @@ static int vdbeRecordCompareString(
     nCmp = MIN( pPKey2->aMem[0].n, nStr );
     res = memcmp(&aKey1[szHdr], pPKey2->aMem[0].z, nCmp);
 
-    if( res==0 ){
+    if( res>0 ){
+      res = pPKey2->r2;
+    }else if( res<0 ){
+      res = pPKey2->r1;
+    }else{
       res = nStr - pPKey2->aMem[0].n;
       if( res==0 ){
         if( pPKey2->nField>1 ){
@@ -80891,10 +82096,6 @@ static int vdbeRecordCompareString(
       }else{
         res = pPKey2->r1;
       }
-    }else if( res>0 ){
-      res = pPKey2->r2;
-    }else{
-      res = pPKey2->r1;
     }
   }
 
@@ -80926,7 +82127,10 @@ SQLITE_PRIVATE RecordCompare sqlite3VdbeFindCompare(UnpackedRecord *p){
   ** header size is (12*5 + 1 + 1) bytes.  */
   if( p->pKeyInfo->nAllField<=13 ){
     int flags = p->aMem[0].flags;
-    if( p->pKeyInfo->aSortOrder[0] ){
+    if( p->pKeyInfo->aSortFlags[0] ){
+      if( p->pKeyInfo->aSortFlags[0] & KEYINFO_ORDER_BIGNULL ){
+        return sqlite3VdbeRecordCompare;
+      }
       p->r1 = 1;
       p->r2 = -1;
     }else{
@@ -80939,7 +82143,9 @@ SQLITE_PRIVATE RecordCompare sqlite3VdbeFindCompare(UnpackedRecord *p){
     testcase( flags & MEM_Real );
     testcase( flags & MEM_Null );
     testcase( flags & MEM_Blob );
-    if( (flags & (MEM_Real|MEM_Null|MEM_Blob))==0 && p->pKeyInfo->aColl[0]==0 ){
+    if( (flags & (MEM_Real|MEM_IntReal|MEM_Null|MEM_Blob))==0
+     && p->pKeyInfo->aColl[0]==0
+    ){
       assert( flags & MEM_Str );
       return vdbeRecordCompareString;
     }
@@ -81173,13 +82379,25 @@ SQLITE_PRIVATE void sqlite3VdbeSetVarmask(Vdbe *v, int iVar){
 ** features such as 'now'.
 */
 SQLITE_PRIVATE int sqlite3NotPureFunc(sqlite3_context *pCtx){
-#ifdef SQLITE_ENABLE_STAT3_OR_STAT4
+  const VdbeOp *pOp;
+#ifdef SQLITE_ENABLE_STAT4
   if( pCtx->pVdbe==0 ) return 1;
 #endif
-  if( pCtx->pVdbe->aOp[pCtx->iOp].opcode==OP_PureFunc ){
-    sqlite3_result_error(pCtx, 
-       "non-deterministic function in index expression or CHECK constraint",
-       -1);
+  pOp = pCtx->pVdbe->aOp + pCtx->iOp;
+  if( pOp->opcode==OP_PureFunc ){
+    const char *zContext;
+    char *zMsg;
+    if( pOp->p5 & NC_IsCheck ){
+      zContext = "a CHECK constraint";
+    }else if( pOp->p5 & NC_GenCol ){
+      zContext = "a generated column";
+    }else{
+      zContext = "an index";
+    }
+    zMsg = sqlite3_mprintf("non-deterministic use of %s() in %s",
+                           pCtx->pFunc->zName, zContext);
+    sqlite3_result_error(pCtx, zMsg, -1);
+    sqlite3_free(zMsg);
     return 0;
   }
   return 1;
@@ -81270,7 +82488,7 @@ SQLITE_PRIVATE void sqlite3VdbePreUpdateHook(
   preupdate.keyinfo.db = db;
   preupdate.keyinfo.enc = ENC(db);
   preupdate.keyinfo.nKeyField = pTab->nCol;
-  preupdate.keyinfo.aSortOrder = (u8*)&fakeSortOrder;
+  preupdate.keyinfo.aSortFlags = (u8*)&fakeSortOrder;
   preupdate.iKey1 = iKey1;
   preupdate.iKey2 = iKey2;
   preupdate.pTab = pTab;
@@ -81529,39 +82747,86 @@ SQLITE_API const void *sqlite3_value_text16le(sqlite3_value *pVal){
 */
 SQLITE_API int sqlite3_value_type(sqlite3_value* pVal){
   static const u8 aType[] = {
-     SQLITE_BLOB,     /* 0x00 */
-     SQLITE_NULL,     /* 0x01 */
-     SQLITE_TEXT,     /* 0x02 */
-     SQLITE_NULL,     /* 0x03 */
-     SQLITE_INTEGER,  /* 0x04 */
-     SQLITE_NULL,     /* 0x05 */
-     SQLITE_INTEGER,  /* 0x06 */
-     SQLITE_NULL,     /* 0x07 */
-     SQLITE_FLOAT,    /* 0x08 */
-     SQLITE_NULL,     /* 0x09 */
-     SQLITE_FLOAT,    /* 0x0a */
-     SQLITE_NULL,     /* 0x0b */
-     SQLITE_INTEGER,  /* 0x0c */
-     SQLITE_NULL,     /* 0x0d */
-     SQLITE_INTEGER,  /* 0x0e */
-     SQLITE_NULL,     /* 0x0f */
-     SQLITE_BLOB,     /* 0x10 */
-     SQLITE_NULL,     /* 0x11 */
-     SQLITE_TEXT,     /* 0x12 */
-     SQLITE_NULL,     /* 0x13 */
-     SQLITE_INTEGER,  /* 0x14 */
-     SQLITE_NULL,     /* 0x15 */
-     SQLITE_INTEGER,  /* 0x16 */
-     SQLITE_NULL,     /* 0x17 */
-     SQLITE_FLOAT,    /* 0x18 */
-     SQLITE_NULL,     /* 0x19 */
-     SQLITE_FLOAT,    /* 0x1a */
-     SQLITE_NULL,     /* 0x1b */
-     SQLITE_INTEGER,  /* 0x1c */
-     SQLITE_NULL,     /* 0x1d */
-     SQLITE_INTEGER,  /* 0x1e */
-     SQLITE_NULL,     /* 0x1f */
+     SQLITE_BLOB,     /* 0x00 (not possible) */
+     SQLITE_NULL,     /* 0x01 NULL */
+     SQLITE_TEXT,     /* 0x02 TEXT */
+     SQLITE_NULL,     /* 0x03 (not possible) */
+     SQLITE_INTEGER,  /* 0x04 INTEGER */
+     SQLITE_NULL,     /* 0x05 (not possible) */
+     SQLITE_INTEGER,  /* 0x06 INTEGER + TEXT */
+     SQLITE_NULL,     /* 0x07 (not possible) */
+     SQLITE_FLOAT,    /* 0x08 FLOAT */
+     SQLITE_NULL,     /* 0x09 (not possible) */
+     SQLITE_FLOAT,    /* 0x0a FLOAT + TEXT */
+     SQLITE_NULL,     /* 0x0b (not possible) */
+     SQLITE_INTEGER,  /* 0x0c (not possible) */
+     SQLITE_NULL,     /* 0x0d (not possible) */
+     SQLITE_INTEGER,  /* 0x0e (not possible) */
+     SQLITE_NULL,     /* 0x0f (not possible) */
+     SQLITE_BLOB,     /* 0x10 BLOB */
+     SQLITE_NULL,     /* 0x11 (not possible) */
+     SQLITE_TEXT,     /* 0x12 (not possible) */
+     SQLITE_NULL,     /* 0x13 (not possible) */
+     SQLITE_INTEGER,  /* 0x14 INTEGER + BLOB */
+     SQLITE_NULL,     /* 0x15 (not possible) */
+     SQLITE_INTEGER,  /* 0x16 (not possible) */
+     SQLITE_NULL,     /* 0x17 (not possible) */
+     SQLITE_FLOAT,    /* 0x18 FLOAT + BLOB */
+     SQLITE_NULL,     /* 0x19 (not possible) */
+     SQLITE_FLOAT,    /* 0x1a (not possible) */
+     SQLITE_NULL,     /* 0x1b (not possible) */
+     SQLITE_INTEGER,  /* 0x1c (not possible) */
+     SQLITE_NULL,     /* 0x1d (not possible) */
+     SQLITE_INTEGER,  /* 0x1e (not possible) */
+     SQLITE_NULL,     /* 0x1f (not possible) */
+     SQLITE_FLOAT,    /* 0x20 INTREAL */
+     SQLITE_NULL,     /* 0x21 (not possible) */
+     SQLITE_TEXT,     /* 0x22 INTREAL + TEXT */
+     SQLITE_NULL,     /* 0x23 (not possible) */
+     SQLITE_FLOAT,    /* 0x24 (not possible) */
+     SQLITE_NULL,     /* 0x25 (not possible) */
+     SQLITE_FLOAT,    /* 0x26 (not possible) */
+     SQLITE_NULL,     /* 0x27 (not possible) */
+     SQLITE_FLOAT,    /* 0x28 (not possible) */
+     SQLITE_NULL,     /* 0x29 (not possible) */
+     SQLITE_FLOAT,    /* 0x2a (not possible) */
+     SQLITE_NULL,     /* 0x2b (not possible) */
+     SQLITE_FLOAT,    /* 0x2c (not possible) */
+     SQLITE_NULL,     /* 0x2d (not possible) */
+     SQLITE_FLOAT,    /* 0x2e (not possible) */
+     SQLITE_NULL,     /* 0x2f (not possible) */
+     SQLITE_BLOB,     /* 0x30 (not possible) */
+     SQLITE_NULL,     /* 0x31 (not possible) */
+     SQLITE_TEXT,     /* 0x32 (not possible) */
+     SQLITE_NULL,     /* 0x33 (not possible) */
+     SQLITE_FLOAT,    /* 0x34 (not possible) */
+     SQLITE_NULL,     /* 0x35 (not possible) */
+     SQLITE_FLOAT,    /* 0x36 (not possible) */
+     SQLITE_NULL,     /* 0x37 (not possible) */
+     SQLITE_FLOAT,    /* 0x38 (not possible) */
+     SQLITE_NULL,     /* 0x39 (not possible) */
+     SQLITE_FLOAT,    /* 0x3a (not possible) */
+     SQLITE_NULL,     /* 0x3b (not possible) */
+     SQLITE_FLOAT,    /* 0x3c (not possible) */
+     SQLITE_NULL,     /* 0x3d (not possible) */
+     SQLITE_FLOAT,    /* 0x3e (not possible) */
+     SQLITE_NULL,     /* 0x3f (not possible) */
   };
+#ifdef SQLITE_DEBUG
+  {
+    int eType = SQLITE_BLOB;
+    if( pVal->flags & MEM_Null ){
+      eType = SQLITE_NULL;
+    }else if( pVal->flags & (MEM_Real|MEM_IntReal) ){
+      eType = SQLITE_FLOAT;
+    }else if( pVal->flags & MEM_Int ){
+      eType = SQLITE_INTEGER;
+    }else if( pVal->flags & MEM_Str ){
+      eType = SQLITE_TEXT;
+    }
+    assert( eType == aType[pVal->flags&MEM_AffMask] );
+  }
+#endif
   return aType[pVal->flags&MEM_AffMask];
 }
 
@@ -81811,6 +83076,21 @@ SQLITE_API void sqlite3_result_error_nomem(sqlite3_context *pCtx){
   sqlite3OomFault(pCtx->pOut->db);
 }
 
+#ifndef SQLITE_UNTESTABLE
+/* Force the INT64 value currently stored as the result to be
+** a MEM_IntReal value.  See the SQLITE_TESTCTRL_RESULT_INTREAL
+** test-control.
+*/
+SQLITE_PRIVATE void sqlite3ResultIntReal(sqlite3_context *pCtx){ 
+  assert( sqlite3_mutex_held(pCtx->pOut->db->mutex) );
+  if( pCtx->pOut->flags & MEM_Int ){
+    pCtx->pOut->flags &= ~MEM_Int;
+    pCtx->pOut->flags |= MEM_IntReal;
+  }
+}
+#endif
+
+
 /*
 ** This function is called after a transaction has been committed. It 
 ** invokes callbacks registered with sqlite3_wal_hook() as required.
@@ -82077,7 +83357,7 @@ SQLITE_API int sqlite3_vtab_nochange(sqlite3_context *p){
 */
 SQLITE_PRIVATE sqlite3_int64 sqlite3StmtCurrentTime(sqlite3_context *p){
   int rc;
-#ifndef SQLITE_ENABLE_STAT3_OR_STAT4
+#ifndef SQLITE_ENABLE_STAT4
   sqlite3_int64 *piTime = &p->pVdbe->iCurrentTime;
   assert( p->pVdbe!=0 );
 #else
@@ -82142,7 +83422,7 @@ SQLITE_API void *sqlite3_get_auxdata(sqlite3_context *pCtx, int iArg){
   AuxData *pAuxData;
 
   assert( sqlite3_mutex_held(pCtx->pOut->db->mutex) );
-#if SQLITE_ENABLE_STAT3_OR_STAT4
+#if SQLITE_ENABLE_STAT4
   if( pCtx->pVdbe==0 ) return 0;
 #else
   assert( pCtx->pVdbe!=0 );
@@ -82176,7 +83456,7 @@ SQLITE_API void sqlite3_set_auxdata(
   Vdbe *pVdbe = pCtx->pVdbe;
 
   assert( sqlite3_mutex_held(pCtx->pOut->db->mutex) );
-#ifdef SQLITE_ENABLE_STAT3_OR_STAT4
+#ifdef SQLITE_ENABLE_STAT4
   if( pVdbe==0 ) goto failed;
 #else
   assert( pVdbe!=0 );
@@ -83064,7 +84344,7 @@ SQLITE_API int sqlite3_preupdate_old(sqlite3 *db, int iIdx, sqlite3_value **ppVa
     goto preupdate_old_out;
   }
   if( p->pPk ){
-    iIdx = sqlite3ColumnOfIndex(p->pPk, iIdx);
+    iIdx = sqlite3TableColumnToIndex(p->pPk, iIdx);
   }
   if( iIdx>=p->pCsr->nField || iIdx<0 ){
     rc = SQLITE_RANGE;
@@ -83097,7 +84377,9 @@ SQLITE_API int sqlite3_preupdate_old(sqlite3 *db, int iIdx, sqlite3_value **ppVa
   }else if( iIdx>=p->pUnpacked->nField ){
     *ppValue = (sqlite3_value *)columnNullValue();
   }else if( p->pTab->aCol[iIdx].affinity==SQLITE_AFF_REAL ){
-    if( pMem->flags & MEM_Int ){
+    if( pMem->flags & (MEM_Int|MEM_IntReal) ){
+      testcase( pMem->flags & MEM_Int );
+      testcase( pMem->flags & MEM_IntReal );
       sqlite3VdbeMemRealify(pMem);
     }
   }
@@ -83152,7 +84434,7 @@ SQLITE_API int sqlite3_preupdate_new(sqlite3 *db, int iIdx, sqlite3_value **ppVa
     goto preupdate_new_out;
   }
   if( p->pPk && p->op!=SQLITE_UPDATE ){
-    iIdx = sqlite3ColumnOfIndex(p->pPk, iIdx);
+    iIdx = sqlite3TableColumnToIndex(p->pPk, iIdx);
   }
   if( iIdx>=p->pCsr->nField || iIdx<0 ){
     rc = SQLITE_RANGE;
@@ -83416,7 +84698,7 @@ SQLITE_PRIVATE char *sqlite3VdbeExpandSql(
       pVar = &p->aVar[idx-1];
       if( pVar->flags & MEM_Null ){
         sqlite3_str_append(&out, "NULL", 4);
-      }else if( pVar->flags & MEM_Int ){
+      }else if( pVar->flags & (MEM_Int|MEM_IntReal) ){
         sqlite3_str_appendf(&out, "%lld", pVar->u.i);
       }else if( pVar->flags & MEM_Real ){
         sqlite3_str_appendf(&out, "%!.15g", pVar->u.r);
@@ -83600,6 +84882,26 @@ SQLITE_API int sqlite3_found_count = 0;
 # define UPDATE_MAX_BLOBSIZE(P)
 #endif
 
+#ifdef SQLITE_DEBUG
+/* This routine provides a convenient place to set a breakpoint during
+** tracing with PRAGMA vdbe_trace=on.  The breakpoint fires right after
+** each opcode is printed.  Variables "pc" (program counter) and pOp are
+** available to add conditionals to the breakpoint.  GDB example:
+**
+**         break test_trace_breakpoint if pc=22
+**
+** Other useful labels for breakpoints include:
+**   test_addop_breakpoint(pc,pOp)
+**   sqlite3CorruptError(lineno)
+**   sqlite3MisuseError(lineno)
+**   sqlite3CantopenError(lineno)
+*/
+static void test_trace_breakpoint(int pc, Op *pOp, Vdbe *v){
+  static int n = 0;
+  n++;
+}
+#endif
+
 /*
 ** Invoke the VDBE coverage callback, if that callback is defined.  This
 ** feature is used for test suite validation only and does not appear an
@@ -83678,14 +84980,6 @@ SQLITE_API int sqlite3_found_count = 0;
   }
 #endif
 
-/*
-** Convert the given register into a string if it isn't one
-** already. Return non-zero if a malloc() fails.
-*/
-#define Stringify(P, enc) \
-   if(((P)->flags&(MEM_Str|MEM_Blob))==0 && sqlite3VdbeMemStringify(P,enc,0)) \
-     { goto no_mem; }
-
 /*
 ** An ephemeral string value (signified by the MEM_Ephem flag) contains
 ** a pointer to a dynamically allocated string where some other entity
@@ -83747,7 +85041,7 @@ static VdbeCursor *allocateCursor(
     ** is clear. Otherwise, if this is an ephemeral cursor created by 
     ** OP_OpenDup, the cursor will not be closed and will still be part
     ** of a BtShared.pCursor list.  */
-    p->apCsr[iCur]->isEphemeral = 0;
+    if( p->apCsr[iCur]->pBtx==0 ) p->apCsr[iCur]->isEphemeral = 0;
     sqlite3VdbeFreeCursor(p, p->apCsr[iCur]);
     p->apCsr[iCur] = 0;
   }
@@ -83767,6 +85061,21 @@ static VdbeCursor *allocateCursor(
   return pCx;
 }
 
+/*
+** The string in pRec is known to look like an integer and to have a
+** floating point value of rValue.  Return true and set *piValue to the
+** integer value if the string is in range to be an integer.  Otherwise,
+** return false.
+*/
+static int alsoAnInt(Mem *pRec, double rValue, i64 *piValue){
+  i64 iValue = (double)rValue;
+  if( sqlite3RealSameAsInt(rValue,iValue) ){
+    *piValue = iValue;
+    return 1;
+  }
+  return 0==sqlite3Atoi64(pRec->z, piValue, pRec->n, pRec->enc);
+}
+
 /*
 ** Try to convert a value into a numeric representation if we can
 ** do so without loss of information.  In other words, if the string
@@ -83784,12 +85093,12 @@ static VdbeCursor *allocateCursor(
 */
 static void applyNumericAffinity(Mem *pRec, int bTryForInt){
   double rValue;
-  i64 iValue;
   u8 enc = pRec->enc;
-  assert( (pRec->flags & (MEM_Str|MEM_Int|MEM_Real))==MEM_Str );
-  if( sqlite3AtoF(pRec->z, &rValue, pRec->n, enc)==0 ) return;
-  if( 0==sqlite3Atoi64(pRec->z, &iValue, pRec->n, enc) ){
-    pRec->u.i = iValue;
+  int rc;
+  assert( (pRec->flags & (MEM_Str|MEM_Int|MEM_Real|MEM_IntReal))==MEM_Str );
+  rc = sqlite3AtoF(pRec->z, &rValue, pRec->n, enc);
+  if( rc<=0 ) return;
+  if( rc==1 && alsoAnInt(pRec, rValue, &pRec->u.i) ){
     pRec->flags |= MEM_Int;
   }else{
     pRec->u.r = rValue;
@@ -83819,6 +85128,7 @@ static void applyNumericAffinity(Mem *pRec, int bTryForInt){
 **    Convert pRec to a text representation.
 **
 ** SQLITE_AFF_BLOB:
+** SQLITE_AFF_NONE:
 **    No-op.  pRec is unchanged.
 */
 static void applyAffinity(
@@ -83843,11 +85153,14 @@ static void applyAffinity(
     ** there is already a string rep, but it is pointless to waste those
     ** CPU cycles. */
     if( 0==(pRec->flags&MEM_Str) ){ /*OPTIMIZATION-IF-FALSE*/
-      if( (pRec->flags&(MEM_Real|MEM_Int)) ){
+      if( (pRec->flags&(MEM_Real|MEM_Int|MEM_IntReal)) ){
+        testcase( pRec->flags & MEM_Int );
+        testcase( pRec->flags & MEM_Real );
+        testcase( pRec->flags & MEM_IntReal );
         sqlite3VdbeMemStringify(pRec, enc, 1);
       }
     }
-    pRec->flags &= ~(MEM_Real|MEM_Int);
+    pRec->flags &= ~(MEM_Real|MEM_Int|MEM_IntReal);
   }
 }
 
@@ -83886,13 +85199,21 @@ SQLITE_PRIVATE void sqlite3ValueApplyAffinity(
 ** accordingly.
 */
 static u16 SQLITE_NOINLINE computeNumericType(Mem *pMem){
-  assert( (pMem->flags & (MEM_Int|MEM_Real))==0 );
+  int rc;
+  sqlite3_int64 ix;
+  assert( (pMem->flags & (MEM_Int|MEM_Real|MEM_IntReal))==0 );
   assert( (pMem->flags & (MEM_Str|MEM_Blob))!=0 );
   ExpandBlob(pMem);
-  if( sqlite3AtoF(pMem->z, &pMem->u.r, pMem->n, pMem->enc)==0 ){
-    return 0;
-  }
-  if( sqlite3Atoi64(pMem->z, &pMem->u.i, pMem->n, pMem->enc)==0 ){
+  rc = sqlite3AtoF(pMem->z, &pMem->u.r, pMem->n, pMem->enc);
+  if( rc<=0 ){
+    if( rc==0 && sqlite3Atoi64(pMem->z, &ix, pMem->n, pMem->enc)<=1 ){
+      pMem->u.i = ix;
+      return MEM_Int;
+    }else{
+      return MEM_Real;
+    }
+  }else if( rc==1 && sqlite3Atoi64(pMem->z, &ix, pMem->n, pMem->enc)==0 ){
+    pMem->u.i = ix;
     return MEM_Int;
   }
   return MEM_Real;
@@ -83906,10 +85227,15 @@ static u16 SQLITE_NOINLINE computeNumericType(Mem *pMem){
 ** But it does set pMem->u.r and pMem->u.i appropriately.
 */
 static u16 numericType(Mem *pMem){
-  if( pMem->flags & (MEM_Int|MEM_Real) ){
-    return pMem->flags & (MEM_Int|MEM_Real);
+  if( pMem->flags & (MEM_Int|MEM_Real|MEM_IntReal) ){
+    testcase( pMem->flags & MEM_Int );
+    testcase( pMem->flags & MEM_Real );
+    testcase( pMem->flags & MEM_IntReal );
+    return pMem->flags & (MEM_Int|MEM_Real|MEM_IntReal);
   }
   if( pMem->flags & (MEM_Str|MEM_Blob) ){
+    testcase( pMem->flags & MEM_Str );
+    testcase( pMem->flags & MEM_Blob );
     return computeNumericType(pMem);
   }
   return 0;
@@ -83920,12 +85246,9 @@ static u16 numericType(Mem *pMem){
 ** Write a nice string representation of the contents of cell pMem
 ** into buffer zBuf, length nBuf.
 */
-SQLITE_PRIVATE void sqlite3VdbeMemPrettyPrint(Mem *pMem, char *zBuf){
-  char *zCsr = zBuf;
+SQLITE_PRIVATE void sqlite3VdbeMemPrettyPrint(Mem *pMem, StrAccum *pStr){
   int f = pMem->flags;
-
   static const char *const encnames[] = {"(X)", "(8)", "(16LE)", "(16BE)"};
-
   if( f&MEM_Blob ){
     int i;
     char c;
@@ -83941,55 +85264,40 @@ SQLITE_PRIVATE void sqlite3VdbeMemPrettyPrint(Mem *pMem, char *zBuf){
     }else{
       c = 's';
     }
-    *(zCsr++) = c;
-    sqlite3_snprintf(100, zCsr, "%d[", pMem->n);
-    zCsr += sqlite3Strlen30(zCsr);
-    for(i=0; i<16 && i<pMem->n; i++){
-      sqlite3_snprintf(100, zCsr, "%02X", ((int)pMem->z[i] & 0xFF));
-      zCsr += sqlite3Strlen30(zCsr);
+    sqlite3_str_appendf(pStr, "%cx[", c);
+    for(i=0; i<25 && i<pMem->n; i++){
+      sqlite3_str_appendf(pStr, "%02X", ((int)pMem->z[i] & 0xFF));
     }
-    for(i=0; i<16 && i<pMem->n; i++){
+    sqlite3_str_appendf(pStr, "|");
+    for(i=0; i<25 && i<pMem->n; i++){
       char z = pMem->z[i];
-      if( z<32 || z>126 ) *zCsr++ = '.';
-      else *zCsr++ = z;
+      sqlite3_str_appendchar(pStr, 1, (z<32||z>126)?'.':z);
     }
-    *(zCsr++) = ']';
+    sqlite3_str_appendf(pStr,"]");
     if( f & MEM_Zero ){
-      sqlite3_snprintf(100, zCsr,"+%dz",pMem->u.nZero);
-      zCsr += sqlite3Strlen30(zCsr);
+      sqlite3_str_appendf(pStr, "+%dz",pMem->u.nZero);
     }
-    *zCsr = '\0';
   }else if( f & MEM_Str ){
-    int j, k;
-    zBuf[0] = ' ';
+    int j;
+    u8 c;
     if( f & MEM_Dyn ){
-      zBuf[1] = 'z';
+      c = 'z';
       assert( (f & (MEM_Static|MEM_Ephem))==0 );
     }else if( f & MEM_Static ){
-      zBuf[1] = 't';
+      c = 't';
       assert( (f & (MEM_Dyn|MEM_Ephem))==0 );
     }else if( f & MEM_Ephem ){
-      zBuf[1] = 'e';
+      c = 'e';
       assert( (f & (MEM_Static|MEM_Dyn))==0 );
     }else{
-      zBuf[1] = 's';
+      c = 's';
     }
-    k = 2;
-    sqlite3_snprintf(100, &zBuf[k], "%d", pMem->n);
-    k += sqlite3Strlen30(&zBuf[k]);
-    zBuf[k++] = '[';
-    for(j=0; j<15 && j<pMem->n; j++){
-      u8 c = pMem->z[j];
-      if( c>=0x20 && c<0x7f ){
-        zBuf[k++] = c;
-      }else{
-        zBuf[k++] = '.';
-      }
+    sqlite3_str_appendf(pStr, " %c%d[", c, pMem->n);
+    for(j=0; j<25 && j<pMem->n; j++){
+      c = pMem->z[j];
+      sqlite3_str_appendchar(pStr, 1, (c>=0x20&&c<=0x7f) ? c : '.');
     }
-    zBuf[k++] = ']';
-    sqlite3_snprintf(100,&zBuf[k], encnames[pMem->enc]);
-    k += sqlite3Strlen30(&zBuf[k]);
-    zBuf[k++] = 0;
+    sqlite3_str_appendf(pStr, "]%s", encnames[pMem->enc]);
   }
 }
 #endif
@@ -84005,29 +85313,48 @@ static void memTracePrint(Mem *p){
     printf(p->flags & MEM_Zero ? " NULL-nochng" : " NULL");
   }else if( (p->flags & (MEM_Int|MEM_Str))==(MEM_Int|MEM_Str) ){
     printf(" si:%lld", p->u.i);
+  }else if( (p->flags & (MEM_IntReal))!=0 ){
+    printf(" ir:%lld", p->u.i);
   }else if( p->flags & MEM_Int ){
     printf(" i:%lld", p->u.i);
 #ifndef SQLITE_OMIT_FLOATING_POINT
   }else if( p->flags & MEM_Real ){
-    printf(" r:%g", p->u.r);
+    printf(" r:%.17g", p->u.r);
 #endif
   }else if( sqlite3VdbeMemIsRowSet(p) ){
     printf(" (rowset)");
   }else{
-    char zBuf[200];
-    sqlite3VdbeMemPrettyPrint(p, zBuf);
-    printf(" %s", zBuf);
+    StrAccum acc;
+    char zBuf[1000];
+    sqlite3StrAccumInit(&acc, 0, zBuf, sizeof(zBuf), 0);
+    sqlite3VdbeMemPrettyPrint(p, &acc);
+    printf(" %s", sqlite3StrAccumFinish(&acc));
   }
   if( p->flags & MEM_Subtype ) printf(" subtype=0x%02x", p->eSubtype);
 }
 static void registerTrace(int iReg, Mem *p){
-  printf("REG[%d] = ", iReg);
+  printf("R[%d] = ", iReg);
   memTracePrint(p);
+  if( p->pScopyFrom ){
+    printf(" <== R[%d]", (int)(p->pScopyFrom - &p[-iReg]));
+  }
   printf("\n");
   sqlite3VdbeCheckMemInvariants(p);
 }
 #endif
 
+#ifdef SQLITE_DEBUG
+/*
+** Show the values of all registers in the virtual machine.  Used for
+** interactive debugging.
+*/
+SQLITE_PRIVATE void sqlite3VdbeRegisterDump(Vdbe *v){
+  int i;
+  for(i=1; i<v->nMem; i++) registerTrace(i, v->aMem+i);
+}
+#endif /* SQLITE_DEBUG */
+
+
 #ifdef SQLITE_DEBUG
 #  define REGISTER_TRACE(R,M) if(db->flags&SQLITE_VdbeTrace)registerTrace(R,M)
 #else
@@ -84056,7 +85383,7 @@ static void registerTrace(int iReg, Mem *p){
 ******************************************************************************
 **
 ** This file contains inline asm code for retrieving "high-performance"
-** counters for x86 class CPUs.
+** counters for x86 and x86_64 class CPUs.
 */
 #ifndef SQLITE_HWTIME_H
 #define SQLITE_HWTIME_H
@@ -84067,8 +85394,9 @@ static void registerTrace(int iReg, Mem *p){
 ** processor and returns that value.  This can be used for high-res
 ** profiling.
 */
-#if (defined(__GNUC__) || defined(_MSC_VER)) && \
-      (defined(i386) || defined(__i386__) || defined(_M_IX86))
+#if !defined(__STRICT_ANSI__) && \
+    (defined(__GNUC__) || defined(_MSC_VER)) && \
+    (defined(i386) || defined(__i386__) || defined(_M_IX86))
 
   #if defined(__GNUC__)
 
@@ -84089,7 +85417,7 @@ static void registerTrace(int iReg, Mem *p){
 
   #endif
 
-#elif (defined(__GNUC__) && defined(__x86_64__))
+#elif !defined(__STRICT_ANSI__) && (defined(__GNUC__) && defined(__x86_64__))
 
   __inline__ sqlite_uint64 sqlite3Hwtime(void){
       unsigned long val;
@@ -84097,7 +85425,7 @@ static void registerTrace(int iReg, Mem *p){
       return val;
   }
  
-#elif (defined(__GNUC__) && defined(__ppc__))
+#elif !defined(__STRICT_ANSI__) && (defined(__GNUC__) && defined(__ppc__))
 
   __inline__ sqlite_uint64 sqlite3Hwtime(void){
       unsigned long long retval;
@@ -84114,14 +85442,13 @@ static void registerTrace(int iReg, Mem *p){
 
 #else
 
-  #error Need implementation of sqlite3Hwtime() for your platform.
-
   /*
-  ** To compile without implementing sqlite3Hwtime() for your platform,
-  ** you can remove the above #error and use the following
-  ** stub function.  You will lose timing support for many
-  ** of the debugging and testing utilities, but it should at
-  ** least compile and run.
+  ** asm() is needed for hardware timing support.  Without asm(),
+  ** disable the sqlite3Hwtime() routine.
+  **
+  ** sqlite3Hwtime() is only used for some obscure debugging
+  ** and analysis configurations, not in any deliverable, so this
+  ** should not be a great loss.
   */
 SQLITE_PRIVATE   sqlite_uint64 sqlite3Hwtime(void){ return ((sqlite_uint64)0); }
 
@@ -84282,6 +85609,7 @@ SQLITE_PRIVATE int sqlite3VdbeExec(
 #ifdef SQLITE_DEBUG
     if( db->flags & SQLITE_VdbeTrace ){
       sqlite3VdbePrintOp(stdout, (int)(pOp - aOp), pOp);
+      test_trace_breakpoint((int)(pOp - aOp),pOp,p);
     }
 #endif
       
@@ -84389,6 +85717,20 @@ SQLITE_PRIVATE int sqlite3VdbeExec(
 ** to the current line should be indented for EXPLAIN output.
 */
 case OP_Goto: {             /* jump */
+
+#ifdef SQLITE_DEBUG
+  /* In debuggging mode, when the p5 flags is set on an OP_Goto, that
+  ** means we should really jump back to the preceeding OP_ReleaseReg
+  ** instruction. */
+  if( pOp->p5 ){
+    assert( pOp->p2 < (int)(pOp - aOp) );
+    assert( pOp->p2 > 1 );
+    pOp = &aOp[pOp->p2 - 2];
+    assert( pOp[1].opcode==OP_ReleaseReg );
+    goto check_for_interrupt;
+  }
+#endif
+
 jump_to_p2_and_check_for_interrupt:
   pOp = &aOp[pOp->p2 - 1];
 
@@ -84687,7 +86029,6 @@ case OP_Real: {            /* same as TK_FLOAT, out2 */
 case OP_String8: {         /* same as TK_STRING, out2 */
   assert( pOp->p4.z!=0 );
   pOut = out2Prerelease(p, pOp);
-  pOp->opcode = OP_String;
   pOp->p1 = sqlite3Strlen30(pOp->p4.z);
 
 #ifndef SQLITE_OMIT_UTF16
@@ -84711,6 +86052,7 @@ case OP_String8: {         /* same as TK_STRING, out2 */
   if( pOp->p1>db->aLimit[SQLITE_LIMIT_LENGTH] ){
     goto too_big;
   }
+  pOp->opcode = OP_String;
   assert( rc==SQLITE_OK );
   /* Fall through to the next case, OP_String */
 }
@@ -84865,8 +86207,13 @@ case OP_Move: {
     memAboutToChange(p, pOut);
     sqlite3VdbeMemMove(pOut, pIn1);
 #ifdef SQLITE_DEBUG
-    if( pOut->pScopyFrom>=&aMem[p1] && pOut->pScopyFrom<pOut ){
-      pOut->pScopyFrom += pOp->p2 - p1;
+    pIn1->pScopyFrom = 0;
+    { int i;
+      for(i=1; i<p->nMem; i++){
+        if( aMem[i].pScopyFrom==pIn1 ){
+          aMem[i].pScopyFrom = pOut;
+        }
+      }
     }
 #endif
     Deephemeralize(pOut);
@@ -85007,6 +86354,14 @@ case OP_ResultRow: {
             || (pMem[i].flags & (MEM_Str|MEM_Blob))==0 );
     sqlite3VdbeMemNulTerminate(&pMem[i]);
     REGISTER_TRACE(pOp->p1+i, &pMem[i]);
+#ifdef SQLITE_DEBUG
+    /* The registers in the result will not be used again when the
+    ** prepared statement restarts.  This is because sqlite3_column()
+    ** APIs might have caused type conversions of made other changes to
+    ** the register values.  Therefore, we can go ahead and break any
+    ** OP_SCopy dependencies. */
+    pMem[i].pScopyFrom = 0;
+#endif
   }
   if( db->mallocFailed ) goto no_mem;
 
@@ -85014,6 +86369,7 @@ case OP_ResultRow: {
     db->xTrace(SQLITE_TRACE_ROW, db->pTraceArg, p, 0);
   }
 
+
   /* Return SQLITE_ROW
   */
   p->pc = (int)(pOp - aOp) + 1;
@@ -85035,33 +86391,57 @@ case OP_ResultRow: {
 ** to avoid a memcpy().
 */
 case OP_Concat: {           /* same as TK_CONCAT, in1, in2, out3 */
-  i64 nByte;
+  i64 nByte;          /* Total size of the output string or blob */
+  u16 flags1;         /* Initial flags for P1 */
+  u16 flags2;         /* Initial flags for P2 */
 
   pIn1 = &aMem[pOp->p1];
   pIn2 = &aMem[pOp->p2];
   pOut = &aMem[pOp->p3];
+  testcase( pIn1==pIn2 );
+  testcase( pOut==pIn2 );
   assert( pIn1!=pOut );
-  if( (pIn1->flags | pIn2->flags) & MEM_Null ){
+  flags1 = pIn1->flags;
+  testcase( flags1 & MEM_Null );
+  testcase( pIn2->flags & MEM_Null );
+  if( (flags1 | pIn2->flags) & MEM_Null ){
     sqlite3VdbeMemSetNull(pOut);
     break;
   }
-  if( ExpandBlob(pIn1) || ExpandBlob(pIn2) ) goto no_mem;
-  Stringify(pIn1, encoding);
-  Stringify(pIn2, encoding);
+  if( (flags1 & (MEM_Str|MEM_Blob))==0 ){
+    if( sqlite3VdbeMemStringify(pIn1,encoding,0) ) goto no_mem;
+    flags1 = pIn1->flags & ~MEM_Str;
+  }else if( (flags1 & MEM_Zero)!=0 ){
+    if( sqlite3VdbeMemExpandBlob(pIn1) ) goto no_mem;
+    flags1 = pIn1->flags & ~MEM_Str;
+  }
+  flags2 = pIn2->flags;
+  if( (flags2 & (MEM_Str|MEM_Blob))==0 ){
+    if( sqlite3VdbeMemStringify(pIn2,encoding,0) ) goto no_mem;
+    flags2 = pIn2->flags & ~MEM_Str;
+  }else if( (flags2 & MEM_Zero)!=0 ){
+    if( sqlite3VdbeMemExpandBlob(pIn2) ) goto no_mem;
+    flags2 = pIn2->flags & ~MEM_Str;
+  }
   nByte = pIn1->n + pIn2->n;
   if( nByte>db->aLimit[SQLITE_LIMIT_LENGTH] ){
     goto too_big;
   }
-  if( sqlite3VdbeMemGrow(pOut, (int)nByte+2, pOut==pIn2) ){
+  if( sqlite3VdbeMemGrow(pOut, (int)nByte+3, pOut==pIn2) ){
     goto no_mem;
   }
   MemSetTypeFlag(pOut, MEM_Str);
   if( pOut!=pIn2 ){
     memcpy(pOut->z, pIn2->z, pIn2->n);
+    assert( (pIn2->flags & MEM_Dyn) == (flags2 & MEM_Dyn) );
+    pIn2->flags = flags2;
   }
   memcpy(&pOut->z[pIn2->n], pIn1->z, pIn1->n);
+  assert( (pIn1->flags & MEM_Dyn) == (flags1 & MEM_Dyn) );
+  pIn1->flags = flags1;
   pOut->z[nByte]=0;
   pOut->z[nByte+1] = 0;
+  pOut->z[nByte+2] = 0;
   pOut->flags |= MEM_Term;
   pOut->n = (int)nByte;
   pOut->enc = encoding;
@@ -85112,7 +86492,6 @@ case OP_Subtract:              /* same as TK_MINUS, in1, in2, out3 */
 case OP_Multiply:              /* same as TK_STAR, in1, in2, out3 */
 case OP_Divide:                /* same as TK_SLASH, in1, in2, out3 */
 case OP_Remainder: {           /* same as TK_REM, in1, in2, out3 */
-  char bIntint;   /* Started out as two integer operands */
   u16 flags;      /* Combined MEM_* flags from both inputs */
   u16 type1;      /* Numeric type of left operand */
   u16 type2;      /* Numeric type of right operand */
@@ -85130,7 +86509,6 @@ case OP_Remainder: {           /* same as TK_REM, in1, in2, out3 */
   if( (type1 & type2 & MEM_Int)!=0 ){
     iA = pIn1->u.i;
     iB = pIn2->u.i;
-    bIntint = 1;
     switch( pOp->opcode ){
       case OP_Add:       if( sqlite3AddInt64(&iB,iA) ) goto fp_math;  break;
       case OP_Subtract:  if( sqlite3SubInt64(&iB,iA) ) goto fp_math;  break;
@@ -85153,7 +86531,6 @@ case OP_Remainder: {           /* same as TK_REM, in1, in2, out3 */
   }else if( (flags & MEM_Null)!=0 ){
     goto arithmetic_result_is_null;
   }else{
-    bIntint = 0;
 fp_math:
     rA = sqlite3VdbeRealValue(pIn1);
     rB = sqlite3VdbeRealValue(pIn2);
@@ -85185,9 +86562,6 @@ fp_math:
     }
     pOut->u.r = rB;
     MemSetTypeFlag(pOut, MEM_Real);
-    if( ((type1|type2)&MEM_Real)==0 && !bIntint ){
-      sqlite3VdbeIntegerAffinity(pOut);
-    }
 #endif
   }
   break;
@@ -85356,8 +86730,11 @@ case OP_MustBeInt: {            /* jump, in1 */
 */
 case OP_RealAffinity: {                  /* in1 */
   pIn1 = &aMem[pOp->p1];
-  if( pIn1->flags & MEM_Int ){
+  if( pIn1->flags & (MEM_Int|MEM_IntReal) ){
+    testcase( pIn1->flags & MEM_Int );
+    testcase( pIn1->flags & MEM_IntReal );
     sqlite3VdbeMemRealify(pIn1);
+    REGISTER_TRACE(pOp->p1, pIn1);
   }
   break;
 }
@@ -85389,9 +86766,11 @@ case OP_Cast: {                  /* in1 */
   pIn1 = &aMem[pOp->p1];
   memAboutToChange(p, pIn1);
   rc = ExpandBlob(pIn1);
-  sqlite3VdbeMemCast(pIn1, pOp->p2, encoding);
-  UPDATE_MAX_BLOBSIZE(pIn1);
   if( rc ) goto abort_due_to_error;
+  rc = sqlite3VdbeMemCast(pIn1, pOp->p2, encoding);
+  if( rc ) goto abort_due_to_error;
+  UPDATE_MAX_BLOBSIZE(pIn1);
+  REGISTER_TRACE(pOp->p1, pIn1);
   break;
 }
 #endif /* SQLITE_OMIT_CAST */
@@ -85548,17 +86927,12 @@ case OP_Ge: {             /* same as TK_GE, jump, in1, in3 */
     affinity = pOp->p5 & SQLITE_AFF_MASK;
     if( affinity>=SQLITE_AFF_NUMERIC ){
       if( (flags1 | flags3)&MEM_Str ){
-        if( (flags1 & (MEM_Int|MEM_Real|MEM_Str))==MEM_Str ){
+        if( (flags1 & (MEM_Int|MEM_IntReal|MEM_Real|MEM_Str))==MEM_Str ){
           applyNumericAffinity(pIn1,0);
-          assert( flags3==pIn3->flags );
-          /* testcase( flags3!=pIn3->flags );
-          ** this used to be possible with pIn1==pIn3, but not since
-          ** the column cache was removed.  The following assignment
-          ** is essentially a no-op.  But, it provides defense-in-depth
-          ** in case our analysis is incorrect, so it is left in. */
+          testcase( flags3!=pIn3->flags );
           flags3 = pIn3->flags;
         }
-        if( (flags3 & (MEM_Int|MEM_Real|MEM_Str))==MEM_Str ){
+        if( (flags3 & (MEM_Int|MEM_IntReal|MEM_Real|MEM_Str))==MEM_Str ){
           applyNumericAffinity(pIn3,0);
         }
       }
@@ -85571,17 +86945,19 @@ case OP_Ge: {             /* same as TK_GE, jump, in1, in3 */
         goto compare_op;
       }
     }else if( affinity==SQLITE_AFF_TEXT ){
-      if( (flags1 & MEM_Str)==0 && (flags1 & (MEM_Int|MEM_Real))!=0 ){
+      if( (flags1 & MEM_Str)==0 && (flags1&(MEM_Int|MEM_Real|MEM_IntReal))!=0 ){
         testcase( pIn1->flags & MEM_Int );
         testcase( pIn1->flags & MEM_Real );
+        testcase( pIn1->flags & MEM_IntReal );
         sqlite3VdbeMemStringify(pIn1, encoding, 1);
         testcase( (flags1&MEM_Dyn) != (pIn1->flags&MEM_Dyn) );
         flags1 = (pIn1->flags & ~MEM_TypeMask) | (flags1 & MEM_TypeMask);
-        assert( pIn1!=pIn3 );
+        if( pIn1==pIn3 ) flags3 = flags1 | MEM_Str;
       }
-      if( (flags3 & MEM_Str)==0 && (flags3 & (MEM_Int|MEM_Real))!=0 ){
+      if( (flags3 & MEM_Str)==0 && (flags3&(MEM_Int|MEM_Real|MEM_IntReal))!=0 ){
         testcase( pIn3->flags & MEM_Int );
         testcase( pIn3->flags & MEM_Real );
+        testcase( pIn3->flags & MEM_IntReal );
         sqlite3VdbeMemStringify(pIn3, encoding, 1);
         testcase( (flags3&MEM_Dyn) != (pIn3->flags&MEM_Dyn) );
         flags3 = (pIn3->flags & ~MEM_TypeMask) | (flags3 & MEM_TypeMask);
@@ -85611,10 +86987,10 @@ compare_op:
   }
 
   /* Undo any changes made by applyAffinity() to the input registers. */
-  assert( (pIn1->flags & MEM_Dyn) == (flags1 & MEM_Dyn) );
-  pIn1->flags = flags1;
   assert( (pIn3->flags & MEM_Dyn) == (flags3 & MEM_Dyn) );
   pIn3->flags = flags3;
+  assert( (pIn1->flags & MEM_Dyn) == (flags1 & MEM_Dyn) );
+  pIn1->flags = flags1;
 
   if( pOp->p5 & SQLITE_STOREP2 ){
     pOut = &aMem[pOp->p2];
@@ -85650,16 +87026,31 @@ compare_op:
 
 /* Opcode: ElseNotEq * P2 * * *
 **
-** This opcode must immediately follow an OP_Lt or OP_Gt comparison operator.
-** If result of an OP_Eq comparison on the same two operands
-** would have be NULL or false (0), then then jump to P2. 
-** If the result of an OP_Eq comparison on the two previous operands
-** would have been true (1), then fall through.
+** This opcode must follow an OP_Lt or OP_Gt comparison operator.  There
+** can be zero or more OP_ReleaseReg opcodes intervening, but no other
+** opcodes are allowed to occur between this instruction and the previous
+** OP_Lt or OP_Gt.  Furthermore, the prior OP_Lt or OP_Gt must have the
+** SQLITE_STOREP2 bit set in the P5 field.
+**
+** If result of an OP_Eq comparison on the same two operands as the
+** prior OP_Lt or OP_Gt would have been NULL or false (0), then then
+** jump to P2.  If the result of an OP_Eq comparison on the two previous
+** operands would have been true (1), then fall through.
 */
 case OP_ElseNotEq: {       /* same as TK_ESCAPE, jump */
-  assert( pOp>aOp );
-  assert( pOp[-1].opcode==OP_Lt || pOp[-1].opcode==OP_Gt );
-  assert( pOp[-1].p5 & SQLITE_STOREP2 );
+
+#ifdef SQLITE_DEBUG
+  /* Verify the preconditions of this opcode - that it follows an OP_Lt or
+  ** OP_Gt with the SQLITE_STOREP2 flag set, with zero or more intervening
+  ** OP_ReleaseReg opcodes */
+  int iAddr;
+  for(iAddr = (int)(pOp - aOp) - 1; ALWAYS(iAddr>=0); iAddr--){
+    if( aOp[iAddr].opcode==OP_ReleaseReg ) continue;
+    assert( aOp[iAddr].opcode==OP_Lt || aOp[iAddr].opcode==OP_Gt );
+    assert( aOp[iAddr].p5 & SQLITE_STOREP2 );
+    break;
+  }
+#endif /* SQLITE_DEBUG */
   VdbeBranchTaken(iCompare!=0, 2);
   if( iCompare!=0 ) goto jump_to_p2;
   break;
@@ -85751,9 +87142,14 @@ case OP_Compare: {
     REGISTER_TRACE(p2+idx, &aMem[p2+idx]);
     assert( i<pKeyInfo->nKeyField );
     pColl = pKeyInfo->aColl[i];
-    bRev = pKeyInfo->aSortOrder[i];
+    bRev = (pKeyInfo->aSortFlags[i] & KEYINFO_ORDER_DESC);
     iCompare = sqlite3MemCompare(&aMem[p1+idx], &aMem[p2+idx], pColl);
     if( iCompare ){
+      if( (pKeyInfo->aSortFlags[i] & KEYINFO_ORDER_BIGNULL) 
+       && ((aMem[p1+idx].flags & MEM_Null) || (aMem[p2+idx].flags & MEM_Null))
+      ){
+        iCompare = -iCompare;
+      }
       if( bRev ) iCompare = -iCompare;
       break;
     }
@@ -86044,11 +87440,6 @@ case OP_Offset: {          /* out3 */
 ** if the P4 argument is a P4_MEM use the value of the P4 argument as
 ** the result.
 **
-** If the OPFLAG_CLEARCACHE bit is set on P5 and P1 is a pseudo-table cursor,
-** then the cache of the cursor is reset prior to extracting the column.
-** The first OP_Column against a pseudo-table after the value of the content
-** register has changed should have this bit set.
-**
 ** If the OPFLAG_LENGTHARG and OPFLAG_TYPEOFARG bits are set on P5 then
 ** the result is guaranteed to only be used as the argument of a length()
 ** or typeof() function, respectively.  The loading of large blobs can be
@@ -86070,7 +87461,9 @@ case OP_Column: {
   u32 t;             /* A type code from the record header */
   Mem *pReg;         /* PseudoTable input register */
 
+  assert( pOp->p1>=0 && pOp->p1<p->nCursor );
   pC = p->apCsr[pOp->p1];
+  assert( pC!=0 );
   p2 = pOp->p2;
 
   /* If the cursor cache is stale (meaning it is not currently point at
@@ -86082,7 +87475,6 @@ case OP_Column: {
   assert( pOp->p3>0 && pOp->p3<=(p->nMem+1 - p->nCursor) );
   pDest = &aMem[pOp->p3];
   memAboutToChange(p, pDest);
-  assert( pOp->p1>=0 && pOp->p1<p->nCursor );
   assert( pC!=0 );
   assert( p2<pC->nField );
   aOffset = pC->aOffset;
@@ -86293,10 +87685,11 @@ case OP_Column: {
       **
       ** Although sqlite3VdbeSerialGet() may read at most 8 bytes from the
       ** buffer passed to it, debugging function VdbeMemPrettyPrint() may
-      ** read up to 16. So 16 bytes of bogus content is supplied.
+      ** read more.  Use the global constant sqlite3CtypeMap[] as the array,
+      ** as that array is 256 bytes long (plenty for VdbeMemPrettyPrint())
+      ** and it begins with a bunch of zeros.
       */
-      static u8 aZero[16];  /* This is the bogus content */
-      sqlite3VdbeSerialGet(aZero, t, pDest);
+      sqlite3VdbeSerialGet((u8*)sqlite3CtypeMap, t, pDest);
     }else{
       rc = sqlite3VdbeMemFromBtree(pC->uc.pCursor, aOffset[p2], len, pDest);
       if( rc!=SQLITE_OK ) goto abort_due_to_error;
@@ -86337,12 +87730,33 @@ case OP_Affinity: {
   assert( pOp->p2>0 );
   assert( zAffinity[pOp->p2]==0 );
   pIn1 = &aMem[pOp->p1];
-  do{
+  while( 1 /*exit-by-break*/ ){
     assert( pIn1 <= &p->aMem[(p->nMem+1 - p->nCursor)] );
-    assert( memIsValid(pIn1) );
-    applyAffinity(pIn1, *(zAffinity++), encoding);
+    assert( zAffinity[0]==SQLITE_AFF_NONE || memIsValid(pIn1) );
+    applyAffinity(pIn1, zAffinity[0], encoding);
+    if( zAffinity[0]==SQLITE_AFF_REAL && (pIn1->flags & MEM_Int)!=0 ){
+      /* When applying REAL affinity, if the result is still an MEM_Int
+      ** that will fit in 6 bytes, then change the type to MEM_IntReal
+      ** so that we keep the high-resolution integer value but know that
+      ** the type really wants to be REAL. */
+      testcase( pIn1->u.i==140737488355328LL );
+      testcase( pIn1->u.i==140737488355327LL );
+      testcase( pIn1->u.i==-140737488355328LL );
+      testcase( pIn1->u.i==-140737488355329LL );
+      if( pIn1->u.i<=140737488355327LL && pIn1->u.i>=-140737488355328LL ){
+        pIn1->flags |= MEM_IntReal;
+        pIn1->flags &= ~MEM_Int;
+      }else{
+        pIn1->u.r = (double)pIn1->u.i;
+        pIn1->flags |= MEM_Real;
+        pIn1->flags &= ~MEM_Int;
+      }
+    }
+    REGISTER_TRACE((int)(pIn1-aMem), pIn1);
+    zAffinity++;
+    if( zAffinity[0]==0 ) break;
     pIn1++;
-  }while( zAffinity[0] );
+  }
   break;
 }
 
@@ -86363,7 +87777,6 @@ case OP_Affinity: {
 ** If P4 is NULL then all index fields have the affinity BLOB.
 */
 case OP_MakeRecord: {
-  u8 *zNewRecord;        /* A buffer to hold the data for the new record */
   Mem *pRec;             /* The new record */
   u64 nData;             /* Number of bytes of data space */
   int nHdr;              /* Number of bytes of header space */
@@ -86376,9 +87789,9 @@ case OP_MakeRecord: {
   int nField;            /* Number of fields in the record */
   char *zAffinity;       /* The affinity string for the record */
   int file_format;       /* File format to use for encoding */
-  int i;                 /* Space used in zNewRecord[] header */
-  int j;                 /* Space used in zNewRecord[] content */
   u32 len;               /* Length of a field */
+  u8 *zHdr;              /* Where to write next byte of the header */
+  u8 *zPayload;          /* Where to write next byte of the payload */
 
   /* Assuming the record contains N fields, the record format looks
   ** like this:
@@ -86417,7 +87830,14 @@ case OP_MakeRecord: {
   if( zAffinity ){
     pRec = pData0;
     do{
-      applyAffinity(pRec++, *(zAffinity++), encoding);
+      applyAffinity(pRec, zAffinity[0], encoding);
+      if( zAffinity[0]==SQLITE_AFF_REAL && (pRec->flags & MEM_Int) ){
+        pRec->flags |= MEM_IntReal;
+        pRec->flags &= ~(MEM_Int);
+      }
+      REGISTER_TRACE((int)(pRec-aMem), pRec);
+      zAffinity++;
+      pRec++;
       assert( zAffinity[0]==0 || pRec<=pLast );
     }while( zAffinity[0] );
   }
@@ -86437,14 +87857,36 @@ case OP_MakeRecord: {
 #endif
 
   /* Loop through the elements that will make up the record to figure
-  ** out how much space is required for the new record.
+  ** out how much space is required for the new record.  After this loop,
+  ** the Mem.uTemp field of each term should hold the serial-type that will
+  ** be used for that term in the generated record:
+  **
+  **   Mem.uTemp value    type
+  **   ---------------    ---------------
+  **      0               NULL
+  **      1               1-byte signed integer
+  **      2               2-byte signed integer
+  **      3               3-byte signed integer
+  **      4               4-byte signed integer
+  **      5               6-byte signed integer
+  **      6               8-byte signed integer
+  **      7               IEEE float
+  **      8               Integer constant 0
+  **      9               Integer constant 1
+  **     10,11            reserved for expansion
+  **    N>=12 and even    BLOB
+  **    N>=13 and odd     text
+  **
+  ** The following additional values are computed:
+  **     nHdr        Number of bytes needed for the record header
+  **     nData       Number of bytes of data space needed for the record
+  **     nZero       Zero bytes at the end of the record
   */
   pRec = pLast;
   do{
     assert( memIsValid(pRec) );
-    serial_type = sqlite3VdbeSerialType(pRec, file_format, &len);
-    if( pRec->flags & MEM_Zero ){
-      if( serial_type==0 ){
+    if( pRec->flags & MEM_Null ){
+      if( pRec->flags & MEM_Zero ){
         /* Values with MEM_Null and MEM_Zero are created by xColumn virtual
         ** table methods that never invoke sqlite3_result_xxxxx() while
         ** computing an unchanging column value in an UPDATE statement.
@@ -86452,19 +87894,83 @@ case OP_MakeRecord: {
         ** so that they can be passed through to xUpdate and have
         ** a true sqlite3_value_nochange(). */
         assert( pOp->p5==OPFLAG_NOCHNG_MAGIC || CORRUPT_DB );
-        serial_type = 10;
-      }else if( nData ){
-        if( sqlite3VdbeMemExpandBlob(pRec) ) goto no_mem;
+        pRec->uTemp = 10;
       }else{
-        nZero += pRec->u.nZero;
-        len -= pRec->u.nZero;
+        pRec->uTemp = 0;
+      }
+      nHdr++;
+    }else if( pRec->flags & (MEM_Int|MEM_IntReal) ){
+      /* Figure out whether to use 1, 2, 4, 6 or 8 bytes. */
+      i64 i = pRec->u.i;
+      u64 uu;
+      testcase( pRec->flags & MEM_Int );
+      testcase( pRec->flags & MEM_IntReal );
+      if( i<0 ){
+        uu = ~i;
+      }else{
+        uu = i;
+      }
+      nHdr++;
+      testcase( uu==127 );               testcase( uu==128 );
+      testcase( uu==32767 );             testcase( uu==32768 );
+      testcase( uu==8388607 );           testcase( uu==8388608 );
+      testcase( uu==2147483647 );        testcase( uu==2147483648 );
+      testcase( uu==140737488355327LL ); testcase( uu==140737488355328LL );
+      if( uu<=127 ){
+        if( (i&1)==i && file_format>=4 ){
+          pRec->uTemp = 8+(u32)uu;
+        }else{
+          nData++;
+          pRec->uTemp = 1;
+        }
+      }else if( uu<=32767 ){
+        nData += 2;
+        pRec->uTemp = 2;
+      }else if( uu<=8388607 ){
+        nData += 3;
+        pRec->uTemp = 3;
+      }else if( uu<=2147483647 ){
+        nData += 4;
+        pRec->uTemp = 4;
+      }else if( uu<=140737488355327LL ){
+        nData += 6;
+        pRec->uTemp = 5;
+      }else{
+        nData += 8;
+        if( pRec->flags & MEM_IntReal ){
+          /* If the value is IntReal and is going to take up 8 bytes to store
+          ** as an integer, then we might as well make it an 8-byte floating
+          ** point value */
+          pRec->u.r = (double)pRec->u.i;
+          pRec->flags &= ~MEM_IntReal;
+          pRec->flags |= MEM_Real;
+          pRec->uTemp = 7;
+        }else{
+          pRec->uTemp = 6;
+        }
+      }
+    }else if( pRec->flags & MEM_Real ){
+      nHdr++;
+      nData += 8;
+      pRec->uTemp = 7;
+    }else{
+      assert( db->mallocFailed || pRec->flags&(MEM_Str|MEM_Blob) );
+      assert( pRec->n>=0 );
+      len = (u32)pRec->n;
+      serial_type = (len*2) + 12 + ((pRec->flags & MEM_Str)!=0);
+      if( pRec->flags & MEM_Zero ){
+        serial_type += pRec->u.nZero*2;
+        if( nData ){
+          if( sqlite3VdbeMemExpandBlob(pRec) ) goto no_mem;
+          len += pRec->u.nZero;
+        }else{
+          nZero += pRec->u.nZero;
+        }
       }
+      nData += len;
+      nHdr += sqlite3VarintLen(serial_type);
+      pRec->uTemp = serial_type;
     }
-    nData += len;
-    testcase( serial_type==127 );
-    testcase( serial_type==128 );
-    nHdr += serial_type<=127 ? 1 : sqlite3VarintLen(serial_type);
-    pRec->uTemp = serial_type;
     if( pRec==pData0 ) break;
     pRec--;
   }while(1);
@@ -86505,34 +88011,34 @@ case OP_MakeRecord: {
       goto no_mem;
     }
   }
-  zNewRecord = (u8 *)pOut->z;
+  pOut->n = (int)nByte;
+  pOut->flags = MEM_Blob;
+  if( nZero ){
+    pOut->u.nZero = nZero;
+    pOut->flags |= MEM_Zero;
+  }
+  UPDATE_MAX_BLOBSIZE(pOut);
+  zHdr = (u8 *)pOut->z;
+  zPayload = zHdr + nHdr;
 
   /* Write the record */
-  i = putVarint32(zNewRecord, nHdr);
-  j = nHdr;
+  zHdr += putVarint32(zHdr, nHdr);
   assert( pData0<=pLast );
   pRec = pData0;
   do{
     serial_type = pRec->uTemp;
     /* EVIDENCE-OF: R-06529-47362 Following the size varint are one or more
     ** additional varints, one per column. */
-    i += putVarint32(&zNewRecord[i], serial_type);            /* serial type */
+    zHdr += putVarint32(zHdr, serial_type);            /* serial type */
     /* EVIDENCE-OF: R-64536-51728 The values for each column in the record
     ** immediately follow the header. */
-    j += sqlite3VdbeSerialPut(&zNewRecord[j], pRec, serial_type); /* content */
+    zPayload += sqlite3VdbeSerialPut(zPayload, pRec, serial_type); /* content */
   }while( (++pRec)<=pLast );
-  assert( i==nHdr );
-  assert( j==nByte );
+  assert( nHdr==(int)(zHdr - (u8*)pOut->z) );
+  assert( nByte==(int)(zPayload - (u8*)pOut->z) );
 
   assert( pOp->p3>0 && pOp->p3<=(p->nMem+1 - p->nCursor) );
-  pOut->n = (int)nByte;
-  pOut->flags = MEM_Blob;
-  if( nZero ){
-    pOut->u.nZero = nZero;
-    pOut->flags |= MEM_Zero;
-  }
   REGISTER_TRACE(pOp->p3, pOut);
-  UPDATE_MAX_BLOBSIZE(pOut);
   break;
 }
 
@@ -86551,19 +88057,20 @@ case OP_Count: {         /* out2 */
   pCrsr = p->apCsr[pOp->p1]->uc.pCursor;
   assert( pCrsr );
   nEntry = 0;  /* Not needed.  Only used to silence a warning. */
-  rc = sqlite3BtreeCount(pCrsr, &nEntry);
+  rc = sqlite3BtreeCount(db, pCrsr, &nEntry);
   if( rc ) goto abort_due_to_error;
   pOut = out2Prerelease(p, pOp);
   pOut->u.i = nEntry;
-  break;
+  goto check_for_interrupt;
 }
 #endif
 
 /* Opcode: Savepoint P1 * * P4 *
 **
 ** Open, release or rollback the savepoint named by parameter P4, depending
-** on the value of P1. To open a new savepoint, P1==0. To release (commit) an
-** existing savepoint, P1==1, or to rollback an existing savepoint P1==2.
+** on the value of P1. To open a new savepoint set P1==0 (SAVEPOINT_BEGIN).
+** To release (commit) an existing savepoint set P1==1 (SAVEPOINT_RELEASE).
+** To rollback an existing savepoint set P1==2 (SAVEPOINT_ROLLBACK).
 */
 case OP_Savepoint: {
   int p1;                         /* Value of P1 operand */
@@ -86631,6 +88138,7 @@ case OP_Savepoint: {
       }
     }
   }else{
+    assert( p1==SAVEPOINT_RELEASE || p1==SAVEPOINT_ROLLBACK );
     iSavepoint = 0;
 
     /* Find the named savepoint. If there is no such savepoint, then an
@@ -86670,8 +88178,12 @@ case OP_Savepoint: {
           p->rc = rc = SQLITE_BUSY;
           goto vdbe_return;
         }
-        db->isTransactionSavepoint = 0;
         rc = p->rc;
+        if( rc ){
+          db->autoCommit = 0;
+        }else{
+          db->isTransactionSavepoint = 0;
+        }
       }else{
         int isSchemaChange;
         iSavepoint = db->nSavepoint - iSavepoint - 1;
@@ -86684,6 +88196,7 @@ case OP_Savepoint: {
             if( rc!=SQLITE_OK ) goto abort_due_to_error;
           }
         }else{
+          assert( p1==SAVEPOINT_RELEASE );
           isSchemaChange = 0;
         }
         for(ii=0; ii<db->nDb; ii++){
@@ -86698,6 +88211,7 @@ case OP_Savepoint: {
           db->mDbFlags |= DBFLAG_SchemaChange;
         }
       }
+      if( rc ) goto abort_due_to_error;
   
       /* Regardless of whether this is a RELEASE or ROLLBACK, destroy all 
       ** savepoints nested inside of the savepoint being operated on. */
@@ -86720,6 +88234,7 @@ case OP_Savepoint: {
           db->nSavepoint--;
         }
       }else{
+        assert( p1==SAVEPOINT_ROLLBACK );
         db->nDeferredCons = pSavepoint->nDeferredCons;
         db->nDeferredImmCons = pSavepoint->nDeferredImmCons;
       }
@@ -86779,7 +88294,6 @@ case OP_AutoCommit: {
       p->rc = rc = SQLITE_BUSY;
       goto vdbe_return;
     }
-    assert( db->nStatement==0 );
     sqlite3CloseSavepoints(db);
     if( p->rc==SQLITE_OK ){
       rc = SQLITE_DONE;
@@ -86796,7 +88310,7 @@ case OP_AutoCommit: {
     rc = SQLITE_ERROR;
     goto abort_due_to_error;
   }
-  break;
+  /*NOTREACHED*/ assert(0);
 }
 
 /* Opcode: Transaction P1 P2 P3 P4 P5
@@ -86860,7 +88374,8 @@ case OP_Transaction: {
       goto abort_due_to_error;
     }
 
-    if( pOp->p2 && p->usesStmtJournal 
+    if( p->usesStmtJournal
+     && pOp->p2
      && (db->autoCommit==0 || db->nVdbeRead>1) 
     ){
       assert( sqlite3BtreeIsInTrans(pBt) );
@@ -87192,6 +88707,7 @@ case OP_OpenDup: {
   VdbeCursor *pCx;      /* The new cursor */
 
   pOrig = p->apCsr[pOp->p2];
+  assert( pOrig );
   assert( pOrig->pBtx!=0 );  /* Only ephemeral cursors can be duplicated */
 
   pCx = allocateCursor(p, pOp->p1, pOrig->nField, -1, CURTYPE_BTREE);
@@ -87255,14 +88771,16 @@ case OP_OpenEphemeral: {
   assert( pOp->p1>=0 );
   assert( pOp->p2>=0 );
   pCx = p->apCsr[pOp->p1];
-  if( pCx ){
+  if( pCx && pCx->pBtx ){
     /* If the ephermeral table is already open, erase all existing content
     ** so that the table is empty again, rather than creating a new table. */
+    assert( pCx->isEphemeral );
+    pCx->seqCount = 0;
+    pCx->cacheStatus = CACHE_STALE;
     rc = sqlite3BtreeClearTable(pCx->pBtx, pCx->pgnoRoot, 0);
   }else{
     pCx = allocateCursor(p, pOp->p1, pOp->p2, -1, CURTYPE_BTREE);
     if( pCx==0 ) goto no_mem;
-    pCx->nullRow = 1;
     pCx->isEphemeral = 1;
     rc = sqlite3BtreeOpen(db->pVfs, 0, db, &pCx->pBtx, 
                           BTREE_OMIT_JOURNAL | BTREE_SINGLE | pOp->p5,
@@ -87298,6 +88816,7 @@ case OP_OpenEphemeral: {
     pCx->isOrdered = (pOp->p5!=BTREE_UNORDERED);
   }
   if( rc ) goto abort_due_to_error;
+  pCx->nullRow = 1;
   break;
 }
 
@@ -87526,7 +89045,10 @@ case OP_SeekGT: {       /* jump, in3, group */
   pC->seekOp = pOp->opcode;
 #endif
 
+  pC->deferredMoveto = 0;
+  pC->cacheStatus = CACHE_STALE;
   if( pC->isTable ){
+    u16 flags3, newType;
     /* The BTREE_SEEK_EQ flag is only set on index cursors */
     assert( sqlite3BtreeCursorHasHint(pC->uc.pCursor, BTREE_SEEK_EQ)==0
               || CORRUPT_DB );
@@ -87535,20 +89057,27 @@ case OP_SeekGT: {       /* jump, in3, group */
     ** blob, or NULL.  But it needs to be an integer before we can do
     ** the seek, so convert it. */
     pIn3 = &aMem[pOp->p3];
-    if( (pIn3->flags & (MEM_Int|MEM_Real|MEM_Str))==MEM_Str ){
+    flags3 = pIn3->flags;
+    if( (flags3 & (MEM_Int|MEM_Real|MEM_IntReal|MEM_Str))==MEM_Str ){
       applyNumericAffinity(pIn3, 0);
     }
-    iKey = sqlite3VdbeIntValue(pIn3);
+    iKey = sqlite3VdbeIntValue(pIn3); /* Get the integer key value */
+    newType = pIn3->flags; /* Record the type after applying numeric affinity */
+    pIn3->flags = flags3;  /* But convert the type back to its original */
 
     /* If the P3 value could not be converted into an integer without
     ** loss of information, then special processing is required... */
-    if( (pIn3->flags & MEM_Int)==0 ){
-      if( (pIn3->flags & MEM_Real)==0 ){
-        /* If the P3 value cannot be converted into any kind of a number,
-        ** then the seek is not possible, so jump to P2 */
-        VdbeBranchTaken(1,2); goto jump_to_p2;
-        break;
-      }
+    if( (newType & (MEM_Int|MEM_IntReal))==0 ){
+      if( (newType & MEM_Real)==0 ){
+        if( (newType & MEM_Null) || oc>=OP_SeekGE ){
+          VdbeBranchTaken(1,2);
+          goto jump_to_p2;
+        }else{
+          rc = sqlite3BtreeLast(pC->uc.pCursor, &res);
+          if( rc!=SQLITE_OK ) goto abort_due_to_error;
+          goto seek_not_found;
+        }
+      }else
 
       /* If the approximation iKey is larger than the actual real search
       ** term, substitute >= for > and < for <=. e.g. if the search term
@@ -87572,7 +89101,7 @@ case OP_SeekGT: {       /* jump, in3, group */
         assert( (OP_SeekLT & 0x0001)==(OP_SeekGE & 0x0001) );
         if( (oc & 0x0001)==(OP_SeekLT & 0x0001) ) oc++;
       }
-    } 
+    }
     rc = sqlite3BtreeMovetoUnpacked(pC->uc.pCursor, 0, (u64)iKey, 0, &res);
     pC->movetoTarget = iKey;  /* Used by OP_Delete */
     if( rc!=SQLITE_OK ){
@@ -87626,8 +89155,6 @@ case OP_SeekGT: {       /* jump, in3, group */
       goto seek_not_found;
     }
   }
-  pC->deferredMoveto = 0;
-  pC->cacheStatus = CACHE_STALE;
 #ifdef SQLITE_TEST
   sqlite3_search_count++;
 #endif
@@ -87682,7 +89209,7 @@ seek_not_found:
 ** Synopsis: seekHit=P2
 **
 ** Set the seekHit flag on cursor P1 to the value in P2.
-** The seekHit flag is used by the IfNoHope opcode.
+* The seekHit flag is used by the IfNoHope opcode.
 **
 ** P1 must be a valid b-tree cursor.  P2 must be a boolean value,
 ** either 0 or 1.
@@ -87697,6 +89224,20 @@ case OP_SeekHit: {
   break;
 }
 
+/* Opcode: IfNotOpen P1 P2 * * *
+** Synopsis: if( !csr[P1] ) goto P2
+**
+** If cursor P1 is not open, jump to instruction P2. Otherwise, fall through.
+*/
+case OP_IfNotOpen: {        /* jump */
+  assert( pOp->p1>=0 && pOp->p1<p->nCursor );
+  VdbeBranchTaken(p->apCsr[pOp->p1]==0, 2);
+  if( !p->apCsr[pOp->p1] ){
+    goto jump_to_p2_and_check_for_interrupt;
+  }
+  break;
+}
+
 /* Opcode: Found P1 P2 P3 P4 *
 ** Synopsis: key=r[P3@P4]
 **
@@ -87927,23 +89468,29 @@ case OP_SeekRowid: {        /* jump, in3 */
   u64 iKey;
 
   pIn3 = &aMem[pOp->p3];
-  if( (pIn3->flags & MEM_Int)==0 ){
-    /* Make sure pIn3->u.i contains a valid integer representation of
-    ** the key value, but do not change the datatype of the register, as
-    ** other parts of the perpared statement might be depending on the
-    ** current datatype. */
-    u16 origFlags = pIn3->flags;
-    int isNotInt;
-    applyAffinity(pIn3, SQLITE_AFF_NUMERIC, encoding);
-    isNotInt = (pIn3->flags & MEM_Int)==0;
-    pIn3->flags = origFlags;
-    if( isNotInt ) goto jump_to_p2;
+  testcase( pIn3->flags & MEM_Int );
+  testcase( pIn3->flags & MEM_IntReal );
+  testcase( pIn3->flags & MEM_Real );
+  testcase( (pIn3->flags & (MEM_Str|MEM_Int))==MEM_Str );
+  if( (pIn3->flags & (MEM_Int|MEM_IntReal))==0 ){
+    /* If pIn3->u.i does not contain an integer, compute iKey as the
+    ** integer value of pIn3.  Jump to P2 if pIn3 cannot be converted
+    ** into an integer without loss of information.  Take care to avoid
+    ** changing the datatype of pIn3, however, as it is used by other
+    ** parts of the prepared statement. */
+    Mem x = pIn3[0];
+    applyAffinity(&x, SQLITE_AFF_NUMERIC, encoding);
+    if( (x.flags & MEM_Int)==0 ) goto jump_to_p2;
+    iKey = x.u.i;
+    goto notExistsWithKey;
   }
   /* Fall through into OP_NotExists */
 case OP_NotExists:          /* jump, in3 */
   pIn3 = &aMem[pOp->p3];
   assert( (pIn3->flags & MEM_Int)!=0 || pOp->opcode==OP_SeekRowid );
   assert( pOp->p1>=0 && pOp->p1<p->nCursor );
+  iKey = pIn3->u.i;
+notExistsWithKey:
   pC = p->apCsr[pOp->p1];
   assert( pC!=0 );
 #ifdef SQLITE_DEBUG
@@ -87954,7 +89501,6 @@ case OP_NotExists:          /* jump, in3 */
   pCrsr = pC->uc.pCursor;
   assert( pCrsr!=0 );
   res = 0;
-  iKey = pIn3->u.i;
   rc = sqlite3BtreeMovetoUnpacked(pCrsr, 0, iKey, 0, &res);
   assert( rc==SQLITE_OK || res==0 );
   pC->movetoTarget = iKey;  /* Used by OP_Delete */
@@ -88180,6 +89726,7 @@ case OP_Insert: {
   pC = p->apCsr[pOp->p1];
   assert( pC!=0 );
   assert( pC->eCurType==CURTYPE_BTREE );
+  assert( pC->deferredMoveto==0 );
   assert( pC->uc.pCursor!=0 );
   assert( (pOp->p5 & OPFLAG_ISNOOP) || pC->isTable );
   assert( pOp->p4type==P4_TABLE || pOp->p4type>=P4_STATIC );
@@ -88297,12 +89844,16 @@ case OP_Delete: {
   sqlite3VdbeIncrWriteCounter(p, pC);
 
 #ifdef SQLITE_DEBUG
-  if( pOp->p4type==P4_TABLE && HasRowid(pOp->p4.pTab) && pOp->p5==0 ){
+  if( pOp->p4type==P4_TABLE
+   && HasRowid(pOp->p4.pTab)
+   && pOp->p5==0
+   && sqlite3BtreeCursorIsValidNN(pC->uc.pCursor)
+  ){
     /* If p5 is zero, the seek operation that positioned the cursor prior to
     ** OP_Delete will have also set the pC->movetoTarget field to the rowid of
     ** the row that is being deleted */
     i64 iKey = sqlite3BtreeIntegerKey(pC->uc.pCursor);
-    assert( pC->movetoTarget==iKey );
+    assert( CORRUPT_DB || pC->movetoTarget==iKey );
   }
 #endif
 
@@ -88836,11 +90387,12 @@ case OP_Next:          /* jump */
   ** The Prev opcode is only used after SeekLT, SeekLE, and Last. */
   assert( pOp->opcode!=OP_Next
        || pC->seekOp==OP_SeekGT || pC->seekOp==OP_SeekGE
-       || pC->seekOp==OP_Rewind || pC->seekOp==OP_Found 
-       || pC->seekOp==OP_NullRow|| pC->seekOp==OP_SeekRowid);
+       || pC->seekOp==OP_Rewind || pC->seekOp==OP_Found
+       || pC->seekOp==OP_NullRow|| pC->seekOp==OP_SeekRowid
+       || pC->seekOp==OP_IfNoHope);
   assert( pOp->opcode!=OP_Prev
        || pC->seekOp==OP_SeekLT || pC->seekOp==OP_SeekLE
-       || pC->seekOp==OP_Last 
+       || pC->seekOp==OP_Last   || pC->seekOp==OP_IfNoHope
        || pC->seekOp==OP_NullRow);
 
   rc = pOp->p4.xAdvance(pC->uc.pCursor, pOp->p3);
@@ -89052,6 +90604,24 @@ case OP_IdxRowid: {           /* out2 */
   break;
 }
 
+/* Opcode: FinishSeek P1 * * * *
+** 
+** If cursor P1 was previously moved via OP_DeferredSeek, complete that
+** seek operation now, without further delay.  If the cursor seek has
+** already occurred, this instruction is a no-op.
+*/
+case OP_FinishSeek: {
+  VdbeCursor *pC;             /* The P1 index cursor */
+
+  assert( pOp->p1>=0 && pOp->p1<p->nCursor );
+  pC = p->apCsr[pOp->p1];
+  if( pC->deferredMoveto ){
+    rc = sqlite3VdbeFinishMoveto(pC);
+    if( rc ) goto abort_due_to_error;
+  }
+  break;
+}
+
 /* Opcode: IdxGE P1 P2 P3 P4 P5
 ** Synopsis: key=r[P3@P4]
 **
@@ -89359,7 +90929,7 @@ case OP_ParseSchema: {
     initData.pzErrMsg = &p->zErrMsg;
     initData.mInitFlags = 0;
     zSql = sqlite3MPrintf(db,
-       "SELECT name, rootpage, sql FROM '%q'.%s WHERE %s ORDER BY rowid",
+       "SELECT*FROM\"%w\".%s WHERE %s ORDER BY rowid",
        db->aDb[iDb].zDbSName, zMaster, pOp->p4.z);
     if( zSql==0 ){
       rc = SQLITE_NOMEM_BKPT;
@@ -89488,7 +91058,7 @@ case OP_IntegrityCk: {
   pIn1 = &aMem[pOp->p1];
   assert( pOp->p5<db->nDb );
   assert( DbMaskTest(p->btreeMask, pOp->p5) );
-  z = sqlite3BtreeIntegrityCheck(db->aDb[pOp->p5].pBt, &aRoot[1], nRoot,
+  z = sqlite3BtreeIntegrityCheck(db, db->aDb[pOp->p5].pBt, &aRoot[1], nRoot,
                                  (int)pnErr->u.i+1, &nErr);
   sqlite3VdbeMemSetNull(pIn1);
   if( nErr==0 ){
@@ -89501,7 +91071,7 @@ case OP_IntegrityCk: {
   }
   UPDATE_MAX_BLOBSIZE(pIn1);
   sqlite3VdbeChangeEncoding(pIn1, encoding);
-  break;
+  goto check_for_interrupt;
 }
 #endif /* SQLITE_OMIT_INTEGRITY_CHECK */
 
@@ -90358,6 +91928,36 @@ case OP_Expire: {
   break;
 }
 
+/* Opcode: CursorLock P1 * * * *
+**
+** Lock the btree to which cursor P1 is pointing so that the btree cannot be
+** written by an other cursor.
+*/
+case OP_CursorLock: {
+  VdbeCursor *pC;
+  assert( pOp->p1>=0 && pOp->p1<p->nCursor );
+  pC = p->apCsr[pOp->p1];
+  assert( pC!=0 );
+  assert( pC->eCurType==CURTYPE_BTREE );
+  sqlite3BtreeCursorPin(pC->uc.pCursor);
+  break;
+}
+
+/* Opcode: CursorUnlock P1 * * * *
+**
+** Unlock the btree to which cursor P1 is pointing so that it can be
+** written by other cursors.
+*/
+case OP_CursorUnlock: {
+  VdbeCursor *pC;
+  assert( pOp->p1>=0 && pOp->p1<p->nCursor );
+  pC = p->apCsr[pOp->p1];
+  assert( pC!=0 );
+  assert( pC->eCurType==CURTYPE_BTREE );
+  sqlite3BtreeCursorUnpin(pC->uc.pCursor);
+  break;
+}
+
 #ifndef SQLITE_OMIT_SHARED_CACHE
 /* Opcode: TableLock P1 P2 P3 P4 *
 ** Synopsis: iDb=P1 root=P2 write=P3
@@ -90602,7 +92202,7 @@ case OP_VColumn: {
   assert( pModule->xColumn );
   memset(&sContext, 0, sizeof(sContext));
   sContext.pOut = pDest;
-  testcase( (pOp->p5 & OPFLAG_NOCHNG)==0 && pOp->p5!=0 );
+  assert( pOp->p5==OPFLAG_NOCHNG || pOp->p5==0 );
   if( pOp->p5 & OPFLAG_NOCHNG ){
     sqlite3VdbeMemSetNull(pDest);
     pDest->flags = MEM_Null|MEM_Zero;
@@ -90827,13 +92427,15 @@ case OP_MaxPgcnt: {            /* out2 */
 }
 #endif
 
-/* Opcode: Function0 P1 P2 P3 P4 P5
+/* Opcode: Function P1 P2 P3 P4 *
 ** Synopsis: r[P3]=func(r[P2@P5])
 **
-** Invoke a user function (P4 is a pointer to a FuncDef object that
-** defines the function) with P5 arguments taken from register P2 and
-** successors.  The result of the function is stored in register P3.
-** Register P3 must not be one of the function inputs.
+** Invoke a user function (P4 is a pointer to an sqlite3_context object that
+** contains a pointer to the function to be run) with arguments taken
+** from register P2 and successors.  The number of arguments is in
+** the sqlite3_context object that P4 points to.
+** The result of the function is stored
+** in register P3.  Register P3 must not be one of the function inputs.
 **
 ** P1 is a 32-bit bitmask indicating whether or not each argument to the 
 ** function was determined to be constant at compile time. If the first
@@ -90842,14 +92444,16 @@ case OP_MaxPgcnt: {            /* out2 */
 ** sqlite3_set_auxdata() API may be safely retained until the next
 ** invocation of this opcode.
 **
-** See also: Function, AggStep, AggFinal
+** See also: AggStep, AggFinal, PureFunc
 */
-/* Opcode: Function P1 P2 P3 P4 P5
+/* Opcode: PureFunc P1 P2 P3 P4 *
 ** Synopsis: r[P3]=func(r[P2@P5])
 **
 ** Invoke a user function (P4 is a pointer to an sqlite3_context object that
-** contains a pointer to the function to be run) with P5 arguments taken
-** from register P2 and successors.  The result of the function is stored
+** contains a pointer to the function to be run) with arguments taken
+** from register P2 and successors.  The number of arguments is in
+** the sqlite3_context object that P4 points to.
+** The result of the function is stored
 ** in register P3.  Register P3 must not be one of the function inputs.
 **
 ** P1 is a 32-bit bitmask indicating whether or not each argument to the 
@@ -90859,40 +92463,16 @@ case OP_MaxPgcnt: {            /* out2 */
 ** sqlite3_set_auxdata() API may be safely retained until the next
 ** invocation of this opcode.
 **
-** SQL functions are initially coded as OP_Function0 with P4 pointing
-** to a FuncDef object.  But on first evaluation, the P4 operand is
-** automatically converted into an sqlite3_context object and the operation
-** changed to this OP_Function opcode.  In this way, the initialization of
-** the sqlite3_context object occurs only once, rather than once for each
-** evaluation of the function.
+** This opcode works exactly like OP_Function.  The only difference is in
+** its name.  This opcode is used in places where the function must be
+** purely non-deterministic.  Some built-in date/time functions can be
+** either determinitic of non-deterministic, depending on their arguments.
+** When those function are used in a non-deterministic way, they will check
+** to see if they were called using OP_PureFunc instead of OP_Function, and
+** if they were, they throw an error.
 **
-** See also: Function0, AggStep, AggFinal
+** See also: AggStep, AggFinal, Function
 */
-case OP_PureFunc0:              /* group */
-case OP_Function0: {            /* group */
-  int n;
-  sqlite3_context *pCtx;
-
-  assert( pOp->p4type==P4_FUNCDEF );
-  n = pOp->p5;
-  assert( pOp->p3>0 && pOp->p3<=(p->nMem+1 - p->nCursor) );
-  assert( n==0 || (pOp->p2>0 && pOp->p2+n<=(p->nMem+1 - p->nCursor)+1) );
-  assert( pOp->p3<pOp->p2 || pOp->p3>=pOp->p2+n );
-  pCtx = sqlite3DbMallocRawNN(db, sizeof(*pCtx) + (n-1)*sizeof(sqlite3_value*));
-  if( pCtx==0 ) goto no_mem;
-  pCtx->pOut = 0;
-  pCtx->pFunc = pOp->p4.pFunc;
-  pCtx->iOp = (int)(pOp - aOp);
-  pCtx->pVdbe = p;
-  pCtx->isError = 0;
-  pCtx->argc = n;
-  pOp->p4type = P4_FUNCCTX;
-  pOp->p4.pCtx = pCtx;
-  assert( OP_PureFunc == OP_PureFunc0+2 );
-  assert( OP_Function == OP_Function0+2 );
-  pOp->opcode += 2;
-  /* Fall through into OP_Function */
-}
 case OP_PureFunc:              /* group */
 case OP_Function: {            /* group */
   int i;
@@ -90907,9 +92487,11 @@ case OP_Function: {            /* group */
   ** reinitializes the relavant parts of the sqlite3_context object */
   pOut = &aMem[pOp->p3];
   if( pCtx->pOut != pOut ){
+    pCtx->pVdbe = p;
     pCtx->pOut = pOut;
     for(i=pCtx->argc-1; i>=0; i--) pCtx->argv[i] = &aMem[pOp->p2+i];
   }
+  assert( pCtx->pVdbe==p );
 
   memAboutToChange(p, pOut);
 #ifdef SQLITE_DEBUG
@@ -91081,6 +92663,55 @@ case OP_Abortable: {
 }
 #endif
 
+#ifdef SQLITE_DEBUG
+/* Opcode:  ReleaseReg   P1 P2 P3 * P5
+** Synopsis: release r[P1@P2] mask P3
+**
+** Release registers from service.  Any content that was in the
+** the registers is unreliable after this opcode completes.
+**
+** The registers released will be the P2 registers starting at P1,
+** except if bit ii of P3 set, then do not release register P1+ii.
+** In other words, P3 is a mask of registers to preserve.
+**
+** Releasing a register clears the Mem.pScopyFrom pointer.  That means
+** that if the content of the released register was set using OP_SCopy,
+** a change to the value of the source register for the OP_SCopy will no longer
+** generate an assertion fault in sqlite3VdbeMemAboutToChange().
+**
+** If P5 is set, then all released registers have their type set
+** to MEM_Undefined so that any subsequent attempt to read the released
+** register (before it is reinitialized) will generate an assertion fault.
+**
+** P5 ought to be set on every call to this opcode.
+** However, there are places in the code generator will release registers
+** before their are used, under the (valid) assumption that the registers
+** will not be reallocated for some other purpose before they are used and
+** hence are safe to release.
+**
+** This opcode is only available in testing and debugging builds.  It is
+** not generated for release builds.  The purpose of this opcode is to help
+** validate the generated bytecode.  This opcode does not actually contribute
+** to computing an answer.
+*/
+case OP_ReleaseReg: {
+  Mem *pMem;
+  int i;
+  u32 constMask;
+  assert( pOp->p1>0 );
+  assert( pOp->p1+pOp->p2<=(p->nMem+1 - p->nCursor)+1 );
+  pMem = &aMem[pOp->p1];
+  constMask = pOp->p3;
+  for(i=0; i<pOp->p2; i++, pMem++){
+    if( i>=32 || (constMask & MASKBIT32(i))==0 ){
+      pMem->pScopyFrom = 0;
+      if( i<32 && pOp->p5 ) MemSetTypeFlag(pMem, MEM_Undefined);
+    }
+  }
+  break;
+}
+#endif
+
 /* Opcode: Noop * * * * *
 **
 ** Do nothing.  This instruction is often useful as a jump
@@ -91132,6 +92763,12 @@ default: {          /* This is really OP_Noop, OP_Explain */
       if( opProperty & OPFLG_OUT3 ){
         registerTrace(pOrigOp->p3, &aMem[pOrigOp->p3]);
       }
+      if( opProperty==0xff ){
+        /* Never happens.  This code exists to avoid a harmless linkage
+        ** warning aboud sqlite3VdbeRegisterDump() being defined but not
+        ** used. */
+        sqlite3VdbeRegisterDump(p);
+      }
     }
 #endif  /* SQLITE_DEBUG */
 #endif  /* NDEBUG */
@@ -91566,11 +93203,12 @@ SQLITE_API int sqlite3_blob_close(sqlite3_blob *pBlob){
   sqlite3 *db;
 
   if( p ){
+    sqlite3_stmt *pStmt = p->pStmt;
     db = p->db;
     sqlite3_mutex_enter(db->mutex);
-    rc = sqlite3_finalize(p->pStmt);
     sqlite3DbFree(db, p);
     sqlite3_mutex_leave(db->mutex);
+    rc = sqlite3_finalize(pStmt);
   }else{
     rc = SQLITE_OK;
   }
@@ -92550,7 +94188,8 @@ static int vdbeSorterCompareText(
       );
     }
   }else{
-    if( pTask->pSorter->pKeyInfo->aSortOrder[0] ){
+    assert( !(pTask->pSorter->pKeyInfo->aSortFlags[0]&KEYINFO_ORDER_BIGNULL) );
+    if( pTask->pSorter->pKeyInfo->aSortFlags[0] ){
       res = res * -1;
     }
   }
@@ -92618,7 +94257,8 @@ static int vdbeSorterCompareInt(
           pTask, pbKey2Cached, pKey1, nKey1, pKey2, nKey2
       );
     }
-  }else if( pTask->pSorter->pKeyInfo->aSortOrder[0] ){
+  }else if( pTask->pSorter->pKeyInfo->aSortFlags[0] ){
+    assert( !(pTask->pSorter->pKeyInfo->aSortFlags[0]&KEYINFO_ORDER_BIGNULL) );
     res = res * -1;
   }
 
@@ -92733,6 +94373,7 @@ SQLITE_PRIVATE int sqlite3VdbeSorterInit(
 
     if( pKeyInfo->nAllField<13 
      && (pKeyInfo->aColl[0]==0 || pKeyInfo->aColl[0]==db->pDfltColl)
+     && (pKeyInfo->aSortFlags[0] & KEYINFO_ORDER_BIGNULL)==0
     ){
       pSorter->typeMask = SORTER_TYPE_INTEGER | SORTER_TYPE_TEXT;
     }
@@ -93114,20 +94755,16 @@ static SorterCompare vdbeSorterGetCompare(VdbeSorter *p){
 */
 static int vdbeSorterSort(SortSubtask *pTask, SorterList *pList){
   int i;
-  SorterRecord **aSlot;
   SorterRecord *p;
   int rc;
+  SorterRecord *aSlot[64];
 
   rc = vdbeSortAllocUnpacked(pTask);
   if( rc!=SQLITE_OK ) return rc;
 
   p = pList->pList;
   pTask->xCompare = vdbeSorterGetCompare(pTask->pSorter);
-
-  aSlot = (SorterRecord **)sqlite3MallocZero(64 * sizeof(SorterRecord *));
-  if( !aSlot ){
-    return SQLITE_NOMEM_BKPT;
-  }
+  memset(aSlot, 0, sizeof(aSlot));
 
   while( p ){
     SorterRecord *pNext;
@@ -93152,13 +94789,12 @@ static int vdbeSorterSort(SortSubtask *pTask, SorterList *pList){
   }
 
   p = 0;
-  for(i=0; i<64; i++){
+  for(i=0; i<ArraySize(aSlot); i++){
     if( aSlot[i]==0 ) continue;
     p = p ? vdbeSorterMerge(pTask, p, aSlot[i]) : aSlot[i];
   }
   pList->pList = p;
 
-  sqlite3_free(aSlot);
   assert( pTask->pUnpacked->errCode==SQLITE_OK 
        || pTask->pUnpacked->errCode==SQLITE_NOMEM 
   );
@@ -93449,13 +95085,16 @@ static int vdbeSorterFlushPMA(VdbeSorter *pSorter){
       rc = vdbeSorterListToPMA(&pSorter->aTask[nWorker], &pSorter->list);
     }else{
       /* Launch a background thread for this operation */
-      u8 *aMem = pTask->list.aMemory;
-      void *pCtx = (void*)pTask;
+      u8 *aMem;
+      void *pCtx;
 
+      assert( pTask!=0 );
       assert( pTask->pThread==0 && pTask->bDone==0 );
       assert( pTask->list.pList==0 );
       assert( pTask->list.aMemory==0 || pSorter->list.aMemory!=0 );
 
+      aMem = pTask->list.aMemory;
+      pCtx = (void*)pTask;
       pSorter->iPrev = (u8)(pTask - pSorter->aTask);
       pTask->list = pSorter->list;
       pSorter->list.pList = 0;
@@ -94579,14 +96218,9 @@ static int memjrnlRead(
   int iChunkOffset;
   FileChunk *pChunk;
 
-#if defined(SQLITE_ENABLE_ATOMIC_WRITE) \
- || defined(SQLITE_ENABLE_BATCH_ATOMIC_WRITE)
   if( (iAmt+iOfst)>p->endpoint.iOffset ){
     return SQLITE_IOERR_SHORT_READ;
   }
-#endif
-
-  assert( (iAmt+iOfst)<=p->endpoint.iOffset );
   assert( p->readpoint.iOffset==0 || p->readpoint.pChunk!=0 );
   if( p->readpoint.iOffset!=iOfst || iOfst==0 ){
     sqlite3_int64 iOff = 0;
@@ -94945,9 +96579,22 @@ SQLITE_PRIVATE int sqlite3JournalSize(sqlite3_vfs *pVfs){
 static int walkWindowList(Walker *pWalker, Window *pList){
   Window *pWin;
   for(pWin=pList; pWin; pWin=pWin->pNextWin){
-    if( sqlite3WalkExprList(pWalker, pWin->pOrderBy) ) return WRC_Abort;
-    if( sqlite3WalkExprList(pWalker, pWin->pPartition) ) return WRC_Abort;
-    if( sqlite3WalkExpr(pWalker, pWin->pFilter) ) return WRC_Abort;
+    int rc;
+    rc = sqlite3WalkExprList(pWalker, pWin->pOrderBy);
+    if( rc ) return WRC_Abort;
+    rc = sqlite3WalkExprList(pWalker, pWin->pPartition);
+    if( rc ) return WRC_Abort;
+    rc = sqlite3WalkExpr(pWalker, pWin->pFilter);
+    if( rc ) return WRC_Abort;
+
+    /* The next two are purely for calls to sqlite3RenameExprUnmap()
+    ** within sqlite3WindowOffsetExpr().  Because of constraints imposed
+    ** by sqlite3WindowOffsetExpr(), they can never fail.  The results do
+    ** not matter anyhow. */
+    rc = sqlite3WalkExpr(pWalker, pWin->pStart);
+    if( NEVER(rc) ) return WRC_Abort;
+    rc = sqlite3WalkExpr(pWalker, pWin->pEnd);
+    if( NEVER(rc) ) return WRC_Abort;
   }
   return WRC_Continue;
 }
@@ -94980,21 +96627,25 @@ static SQLITE_NOINLINE int walkExpr(Walker *pWalker, Expr *pExpr){
     rc = pWalker->xExprCallback(pWalker, pExpr);
     if( rc ) return rc & WRC_Abort;
     if( !ExprHasProperty(pExpr,(EP_TokenOnly|EP_Leaf)) ){
+      assert( pExpr->x.pList==0 || pExpr->pRight==0 );
       if( pExpr->pLeft && walkExpr(pWalker, pExpr->pLeft) ) return WRC_Abort;
-       assert( pExpr->x.pList==0 || pExpr->pRight==0 );
       if( pExpr->pRight ){
+        assert( !ExprHasProperty(pExpr, EP_WinFunc) );
         pExpr = pExpr->pRight;
         continue;
       }else if( ExprHasProperty(pExpr, EP_xIsSelect) ){
+        assert( !ExprHasProperty(pExpr, EP_WinFunc) );
         if( sqlite3WalkSelect(pWalker, pExpr->x.pSelect) ) return WRC_Abort;
-      }else if( pExpr->x.pList ){
-        if( sqlite3WalkExprList(pWalker, pExpr->x.pList) ) return WRC_Abort;
-      }
+      }else{
+        if( pExpr->x.pList ){
+          if( sqlite3WalkExprList(pWalker, pExpr->x.pList) ) return WRC_Abort;
+        }
 #ifndef SQLITE_OMIT_WINDOWFUNC
-      if( ExprHasProperty(pExpr, EP_WinFunc) ){
-        if( walkWindowList(pWalker, pExpr->y.pWin) ) return WRC_Abort;
-      }
+        if( ExprHasProperty(pExpr, EP_WinFunc) ){
+          if( walkWindowList(pWalker, pExpr->y.pWin) ) return WRC_Abort;
+        }
 #endif
+      }
     }
     break;
   }
@@ -95036,8 +96687,9 @@ SQLITE_PRIVATE int sqlite3WalkSelectExpr(Walker *pWalker, Select *p){
   {
     Parse *pParse = pWalker->pParse;
     if( pParse && IN_RENAME_OBJECT ){
+      /* The following may return WRC_Abort if there are unresolvable
+      ** symbols (e.g. a table that does not exist) in a window definition. */
       int rc = walkWindowList(pWalker, p->pWinDefn);
-      assert( rc==WRC_Continue );
       return rc;
     }
   }
@@ -95209,6 +96861,13 @@ static void resolveAlias(
       pExpr->u.zToken = sqlite3DbStrDup(db, pExpr->u.zToken);
       pExpr->flags |= EP_MemToken;
     }
+    if( ExprHasProperty(pExpr, EP_WinFunc) ){
+      if( pExpr->y.pWin!=0 ){
+        pExpr->y.pWin->pOwner = pExpr;
+      }else{
+        assert( db->mallocFailed );
+      }
+    }
     sqlite3DbFree(db, pDup);
   }
   ExprSetProperty(pExpr, EP_Alias);
@@ -95238,13 +96897,16 @@ static int nameInUsingClause(IdList *pUsing, const char *zCol){
 ** and zCol.  If any of zDb, zTab, and zCol are NULL then those fields will
 ** match anything.
 */
-SQLITE_PRIVATE int sqlite3MatchSpanName(
-  const char *zSpan,
+SQLITE_PRIVATE int sqlite3MatchEName(
+  const struct ExprList_item *pItem,
   const char *zCol,
   const char *zTab,
   const char *zDb
 ){
   int n;
+  const char *zSpan;
+  if( NEVER(pItem->eEName!=ENAME_TAB) ) return 0;
+  zSpan = pItem->zEName;
   for(n=0; ALWAYS(zSpan[n]) && zSpan[n]!='.'; n++){}
   if( zDb && (sqlite3StrNICmp(zSpan, zDb, n)!=0 || zDb[n]!=0) ){
     return 0;
@@ -95261,6 +96923,23 @@ SQLITE_PRIVATE int sqlite3MatchSpanName(
   return 1;
 }
 
+/*
+** Return TRUE if the double-quoted string  mis-feature should be supported.
+*/
+static int areDoubleQuotedStringsEnabled(sqlite3 *db, NameContext *pTopNC){
+  if( db->init.busy ) return 1;  /* Always support for legacy schemas */
+  if( pTopNC->ncFlags & NC_IsDDL ){
+    /* Currently parsing a DDL statement */
+    if( sqlite3WritableSchema(db) && (db->flags & SQLITE_DqsDML)!=0 ){
+      return 1;
+    }
+    return (db->flags & SQLITE_DqsDDL)!=0;
+  }else{
+    /* Currently parsing a DML statement */
+    return (db->flags & SQLITE_DqsDML)!=0;
+  }
+}
+
 /*
 ** Given the name of a column of the form X.Y.Z or Y.Z or just Z, look up
 ** that name in the set of source tables in pSrcList and make the pExpr 
@@ -95356,7 +97035,7 @@ static int lookupName(
           int hit = 0;
           pEList = pItem->pSelect->pEList;
           for(j=0; j<pEList->nExpr; j++){
-            if( sqlite3MatchSpanName(pEList->a[j].zSpan, zCol, zTab, zDb) ){
+            if( sqlite3MatchEName(&pEList->a[j], zCol, zTab, zDb) ){
               cnt++;
               cntTab = 2;
               pMatch = pItem;
@@ -95477,7 +97156,7 @@ static int lookupName(
           {
 #ifndef SQLITE_OMIT_TRIGGER
             if( iCol<0 ){
-              pExpr->affinity = SQLITE_AFF_INTEGER;
+              pExpr->affExpr = SQLITE_AFF_INTEGER;
             }else if( pExpr->iTable==0 ){
               testcase( iCol==31 );
               testcase( iCol==32 );
@@ -95503,13 +97182,13 @@ static int lookupName(
     if( cnt==0
      && cntTab==1
      && pMatch
-     && (pNC->ncFlags & NC_IdxExpr)==0
+     && (pNC->ncFlags & (NC_IdxExpr|NC_GenCol))==0
      && sqlite3IsRowid(zCol)
      && VisibleRowid(pMatch->pTab)
     ){
       cnt = 1;
       pExpr->iColumn = -1;
-      pExpr->affinity = SQLITE_AFF_INTEGER;
+      pExpr->affExpr = SQLITE_AFF_INTEGER;
     }
 
     /*
@@ -95537,8 +97216,10 @@ static int lookupName(
       pEList = pNC->uNC.pEList;
       assert( pEList!=0 );
       for(j=0; j<pEList->nExpr; j++){
-        char *zAs = pEList->a[j].zName;
-        if( zAs!=0 && sqlite3StrICmp(zAs, zCol)==0 ){
+        char *zAs = pEList->a[j].zEName;
+        if( pEList->a[j].eEName==ENAME_NAME
+         && sqlite3_stricmp(zAs, zCol)==0
+        ){
           Expr *pOrig;
           assert( pExpr->pLeft==0 && pExpr->pRight==0 );
           assert( pExpr->x.pList==0 );
@@ -95548,7 +97229,9 @@ static int lookupName(
             sqlite3ErrorMsg(pParse, "misuse of aliased aggregate %s", zAs);
             return WRC_Abort;
           }
-          if( (pNC->ncFlags&NC_AllowWin)==0 && ExprHasProperty(pOrig, EP_Win) ){
+          if( ExprHasProperty(pOrig, EP_Win)
+           && ((pNC->ncFlags&NC_AllowWin)==0 || pNC!=pTopNC )
+          ){
             sqlite3ErrorMsg(pParse, "misuse of aliased window function %s",zAs);
             return WRC_Abort;
           }
@@ -95589,7 +97272,9 @@ static int lookupName(
   */
   if( cnt==0 && zTab==0 ){
     assert( pExpr->op==TK_ID );
-    if( ExprHasProperty(pExpr,EP_DblQuoted) ){
+    if( ExprHasProperty(pExpr,EP_DblQuoted)
+     && areDoubleQuotedStringsEnabled(db, pTopNC)
+    ){
       /* If a double-quoted identifier does not match any known column name,
       ** then treat it as a string.
       **
@@ -95638,18 +97323,35 @@ static int lookupName(
 
   /* If a column from a table in pSrcList is referenced, then record
   ** this fact in the pSrcList.a[].colUsed bitmask.  Column 0 causes
-  ** bit 0 to be set.  Column 1 sets bit 1.  And so forth.  If the
-  ** column number is greater than the number of bits in the bitmask
-  ** then set the high-order bit of the bitmask.
+  ** bit 0 to be set.  Column 1 sets bit 1.  And so forth.  Bit 63 is
+  ** set if the 63rd or any subsequent column is used.
+  **
+  ** The colUsed mask is an optimization used to help determine if an
+  ** index is a covering index.  The correct answer is still obtained
+  ** if the mask contains extra set bits.  However, it is important to
+  ** avoid setting bits beyond the maximum column number of the table.
+  ** (See ticket [b92e5e8ec2cdbaa1]).
+  **
+  ** If a generated column is referenced, set bits for every column
+  ** of the table.
   */
   if( pExpr->iColumn>=0 && pMatch!=0 ){
     int n = pExpr->iColumn;
-    testcase( n==BMS-1 );
-    if( n>=BMS ){
-      n = BMS-1;
-    }
+    Table *pExTab = pExpr->y.pTab;
+    assert( pExTab!=0 );
     assert( pMatch->iCursor==pExpr->iTable );
-    pMatch->colUsed |= ((Bitmask)1)<<n;
+    if( (pExTab->tabFlags & TF_HasGenerated)!=0
+     && (pExTab->aCol[n].colFlags & COLFLAG_GENERATED)!=0 
+    ){
+      testcase( pExTab->nCol==BMS-1 );
+      testcase( pExTab->nCol==BMS );
+      pMatch->colUsed = pExTab->nCol>=BMS ? ALLBITS : MASKBIT(pExTab->nCol)-1;
+    }else{
+      testcase( n==BMS-1 );
+      testcase( n==BMS );
+      if( n>=BMS ) n = BMS-1;
+      pMatch->colUsed |= ((Bitmask)1)<<n;
+    }
   }
 
   /* Clean up and return
@@ -95688,15 +97390,23 @@ SQLITE_PRIVATE Expr *sqlite3CreateColumnExpr(sqlite3 *db, SrcList *pSrc, int iSr
   Expr *p = sqlite3ExprAlloc(db, TK_COLUMN, 0, 0);
   if( p ){
     struct SrcList_item *pItem = &pSrc->a[iSrc];
-    p->y.pTab = pItem->pTab;
+    Table *pTab = p->y.pTab = pItem->pTab;
     p->iTable = pItem->iCursor;
     if( p->y.pTab->iPKey==iCol ){
       p->iColumn = -1;
     }else{
       p->iColumn = (ynVar)iCol;
-      testcase( iCol==BMS );
-      testcase( iCol==BMS-1 );
-      pItem->colUsed |= ((Bitmask)1)<<(iCol>=BMS ? BMS-1 : iCol);
+      if( (pTab->tabFlags & TF_HasGenerated)!=0
+       && (pTab->aCol[iCol].colFlags & COLFLAG_GENERATED)!=0
+      ){
+        testcase( pTab->nCol==63 );
+        testcase( pTab->nCol==64 );
+        pItem->colUsed = pTab->nCol>=64 ? ALLBITS : MASKBIT(pTab->nCol)-1;
+      }else{
+        testcase( iCol==BMS );
+        testcase( iCol==BMS-1 );
+        pItem->colUsed |= ((Bitmask)1)<<(iCol>=BMS ? BMS-1 : iCol);
+      }
     }
   }
   return p;
@@ -95705,23 +97415,39 @@ SQLITE_PRIVATE Expr *sqlite3CreateColumnExpr(sqlite3 *db, SrcList *pSrc, int iSr
 /*
 ** Report an error that an expression is not valid for some set of
 ** pNC->ncFlags values determined by validMask.
-*/
-static void notValid(
-  Parse *pParse,       /* Leave error message here */
-  NameContext *pNC,    /* The name context */
-  const char *zMsg,    /* Type of error */
-  int validMask        /* Set of contexts for which prohibited */
-){
-  assert( (validMask&~(NC_IsCheck|NC_PartIdx|NC_IdxExpr))==0 );
-  if( (pNC->ncFlags & validMask)!=0 ){
-    const char *zIn = "partial index WHERE clauses";
-    if( pNC->ncFlags & NC_IdxExpr )      zIn = "index expressions";
+**
+** static void notValid(
+**   Parse *pParse,       // Leave error message here
+**   NameContext *pNC,    // The name context 
+**   const char *zMsg,    // Type of error
+**   int validMask,       // Set of contexts for which prohibited
+**   Expr *pExpr          // Invalidate this expression on error
+** ){...}
+**
+** As an optimization, since the conditional is almost always false
+** (because errors are rare), the conditional is moved outside of the
+** function call using a macro.
+*/
+static void notValidImpl(
+   Parse *pParse,       /* Leave error message here */
+   NameContext *pNC,    /* The name context */
+   const char *zMsg,    /* Type of error */
+   Expr *pExpr          /* Invalidate this expression on error */
+){
+  const char *zIn = "partial index WHERE clauses";
+  if( pNC->ncFlags & NC_IdxExpr )      zIn = "index expressions";
 #ifndef SQLITE_OMIT_CHECK
-    else if( pNC->ncFlags & NC_IsCheck ) zIn = "CHECK constraints";
+  else if( pNC->ncFlags & NC_IsCheck ) zIn = "CHECK constraints";
 #endif
-    sqlite3ErrorMsg(pParse, "%s prohibited in %s", zMsg, zIn);
-  }
+#ifndef SQLITE_OMIT_GENERATED_COLUMNS
+  else if( pNC->ncFlags & NC_GenCol ) zIn = "generated columns";
+#endif
+  sqlite3ErrorMsg(pParse, "%s prohibited in %s", zMsg, zIn);
+  if( pExpr ) pExpr->op = TK_NULL;
 }
+#define sqlite3ResolveNotValid(P,N,M,X,E) \
+  assert( ((X)&~(NC_IsCheck|NC_PartIdx|NC_IdxExpr|NC_GenCol))==0 ); \
+  if( ((N)->ncFlags & (X))!=0 ) notValidImpl(P,N,M,E);
 
 /*
 ** Expression p should encode a floating point value between 1.0 and 0.0.
@@ -95783,7 +97509,7 @@ static int resolveExprStep(Walker *pWalker, Expr *pExpr){
       pExpr->y.pTab = pItem->pTab;
       pExpr->iTable = pItem->iCursor;
       pExpr->iColumn = -1;
-      pExpr->affinity = SQLITE_AFF_INTEGER;
+      pExpr->affExpr = SQLITE_AFF_INTEGER;
       break;
     }
 #endif /* defined(SQLITE_ENABLE_UPDATE_DELETE_LIMIT)
@@ -95810,7 +97536,10 @@ static int resolveExprStep(Walker *pWalker, Expr *pExpr){
         zColumn = pExpr->u.zToken;
       }else{
         Expr *pLeft = pExpr->pLeft;
-        notValid(pParse, pNC, "the \".\" operator", NC_IdxExpr);
+        testcase( pNC->ncFlags & NC_IdxExpr );
+        testcase( pNC->ncFlags & NC_GenCol );
+        sqlite3ResolveNotValid(pParse, pNC, "the \".\" operator",
+                               NC_IdxExpr|NC_GenCol, 0);
         pRight = pExpr->pRight;
         if( pRight->op==TK_ID ){
           zDb = 0;
@@ -95843,7 +97572,9 @@ static int resolveExprStep(Walker *pWalker, Expr *pExpr){
       FuncDef *pDef;              /* Information about the function */
       u8 enc = ENC(pParse->db);   /* The database encoding */
       int savedAllowFlags = (pNC->ncFlags & (NC_AllowAgg | NC_AllowWin));
-
+#ifndef SQLITE_OMIT_WINDOWFUNC
+      Window *pWin = (IsWindowFunc(pExpr) ? pExpr->y.pWin : 0);
+#endif
       assert( !ExprHasProperty(pExpr, EP_xIsSelect) );
       zId = pExpr->u.zToken;
       nId = sqlite3Strlen30(zId);
@@ -95858,7 +97589,7 @@ static int resolveExprStep(Walker *pWalker, Expr *pExpr){
       }else{
         is_agg = pDef->xFinalize!=0;
         if( pDef->funcFlags & SQLITE_FUNC_UNLIKELY ){
-          ExprSetProperty(pExpr, EP_Unlikely|EP_Skip);
+          ExprSetProperty(pExpr, EP_Unlikely);
           if( n==2 ){
             pExpr->iTable = exprProbability(pList->a[1].pExpr);
             if( pExpr->iTable<0 ){
@@ -95897,24 +97628,39 @@ static int resolveExprStep(Walker *pWalker, Expr *pExpr){
         if( pDef->funcFlags & (SQLITE_FUNC_CONSTANT|SQLITE_FUNC_SLOCHNG) ){
           /* For the purposes of the EP_ConstFunc flag, date and time
           ** functions and other functions that change slowly are considered
-          ** constant because they are constant for the duration of one query */
+          ** constant because they are constant for the duration of one query.
+          ** This allows them to be factored out of inner loops. */
           ExprSetProperty(pExpr,EP_ConstFunc);
         }
         if( (pDef->funcFlags & SQLITE_FUNC_CONSTANT)==0 ){
-          /* Date/time functions that use 'now', and other functions like
+          /* Clearly non-deterministic functions like random(), but also
+          ** date/time functions that use 'now', and other functions like
           ** sqlite_version() that might change over time cannot be used
-          ** in an index. */
-          notValid(pParse, pNC, "non-deterministic functions",
-                   NC_IdxExpr|NC_PartIdx);
+          ** in an index or generated column.  Curiously, they can be used
+          ** in a CHECK constraint.  SQLServer, MySQL, and PostgreSQL all
+          ** all this. */
+          sqlite3ResolveNotValid(pParse, pNC, "non-deterministic functions",
+                                 NC_IdxExpr|NC_PartIdx|NC_GenCol, 0);
+        }else{
+          assert( (NC_SelfRef & 0xff)==NC_SelfRef ); /* Must fit in 8 bits */
+          pExpr->op2 = pNC->ncFlags & NC_SelfRef;
+          if( pNC->ncFlags & NC_FromDDL ) ExprSetProperty(pExpr, EP_FromDDL);
         }
         if( (pDef->funcFlags & SQLITE_FUNC_INTERNAL)!=0
          && pParse->nested==0
-         && sqlite3Config.bInternalFunctions==0
+         && (pParse->db->mDbFlags & DBFLAG_InternalFunc)==0
         ){
           /* Internal-use-only functions are disallowed unless the
-          ** SQL is being compiled using sqlite3NestedParse() */
+          ** SQL is being compiled using sqlite3NestedParse() or
+          ** the SQLITE_TESTCTRL_INTERNAL_FUNCTIONS test-control has be
+          ** used to activate internal functionsn for testing purposes */
           no_such_func = 1;
           pDef = 0;
+        }else
+        if( (pDef->funcFlags & (SQLITE_FUNC_DIRECT|SQLITE_FUNC_UNSAFE))!=0
+         && !IN_RENAME_OBJECT
+        ){
+          sqlite3ExprFunctionUsable(pParse, pExpr, pDef);
         }
       }
 
@@ -95924,18 +97670,18 @@ static int resolveExprStep(Walker *pWalker, Expr *pExpr){
           || (pDef->xValue==0 && pDef->xInverse==0)
           || (pDef->xValue && pDef->xInverse && pDef->xSFunc && pDef->xFinalize)
         );
-        if( pDef && pDef->xValue==0 && ExprHasProperty(pExpr, EP_WinFunc) ){
+        if( pDef && pDef->xValue==0 && pWin ){
           sqlite3ErrorMsg(pParse, 
               "%.*s() may not be used as a window function", nId, zId
           );
           pNC->nErr++;
         }else if( 
               (is_agg && (pNC->ncFlags & NC_AllowAgg)==0)
-           || (is_agg && (pDef->funcFlags&SQLITE_FUNC_WINDOW) && !pExpr->y.pWin)
-           || (is_agg && pExpr->y.pWin && (pNC->ncFlags & NC_AllowWin)==0)
+           || (is_agg && (pDef->funcFlags&SQLITE_FUNC_WINDOW) && !pWin)
+           || (is_agg && pWin && (pNC->ncFlags & NC_AllowWin)==0)
         ){
           const char *zType;
-          if( (pDef->funcFlags & SQLITE_FUNC_WINDOW) || pExpr->y.pWin ){
+          if( (pDef->funcFlags & SQLITE_FUNC_WINDOW) || pWin ){
             zType = "window";
           }else{
             zType = "aggregate";
@@ -95963,32 +97709,44 @@ static int resolveExprStep(Walker *pWalker, Expr *pExpr){
                nId, zId);
           pNC->nErr++;
         }
+#ifndef SQLITE_OMIT_WINDOWFUNC
+        else if( is_agg==0 && ExprHasProperty(pExpr, EP_WinFunc) ){
+          sqlite3ErrorMsg(pParse, 
+              "FILTER may not be used with non-aggregate %.*s()", 
+              nId, zId
+          );
+          pNC->nErr++;
+        }
+#endif
         if( is_agg ){
           /* Window functions may not be arguments of aggregate functions.
           ** Or arguments of other window functions. But aggregate functions
           ** may be arguments for window functions.  */
 #ifndef SQLITE_OMIT_WINDOWFUNC
-          pNC->ncFlags &= ~(NC_AllowWin | (!pExpr->y.pWin ? NC_AllowAgg : 0));
+          pNC->ncFlags &= ~(NC_AllowWin | (!pWin ? NC_AllowAgg : 0));
 #else
           pNC->ncFlags &= ~NC_AllowAgg;
 #endif
         }
       }
+#ifndef SQLITE_OMIT_WINDOWFUNC
+      else if( ExprHasProperty(pExpr, EP_WinFunc) ){
+        is_agg = 1;
+      }
+#endif
       sqlite3WalkExprList(pWalker, pList);
       if( is_agg ){
 #ifndef SQLITE_OMIT_WINDOWFUNC
-        if( pExpr->y.pWin ){
+        if( pWin ){
           Select *pSel = pNC->pWinSelect;
-          sqlite3WindowUpdate(pParse, pSel->pWinDefn, pExpr->y.pWin, pDef);
-          sqlite3WalkExprList(pWalker, pExpr->y.pWin->pPartition);
-          sqlite3WalkExprList(pWalker, pExpr->y.pWin->pOrderBy);
-          sqlite3WalkExpr(pWalker, pExpr->y.pWin->pFilter);
-          if( 0==pSel->pWin 
-           || 0==sqlite3WindowCompare(pParse, pSel->pWin, pExpr->y.pWin) 
-          ){
-            pExpr->y.pWin->pNextWin = pSel->pWin;
-            pSel->pWin = pExpr->y.pWin;
+          assert( pWin==pExpr->y.pWin );
+          if( IN_RENAME_OBJECT==0 ){
+            sqlite3WindowUpdate(pParse, pSel ? pSel->pWinDefn : 0, pWin, pDef);
           }
+          sqlite3WalkExprList(pWalker, pWin->pPartition);
+          sqlite3WalkExprList(pWalker, pWin->pOrderBy);
+          sqlite3WalkExpr(pWalker, pWin->pFilter);
+          sqlite3WindowLink(pSel, pWin);
           pNC->ncFlags |= NC_HasWin;
         }else
 #endif /* SQLITE_OMIT_WINDOWFUNC */
@@ -95996,12 +97754,17 @@ static int resolveExprStep(Walker *pWalker, Expr *pExpr){
           NameContext *pNC2 = pNC;
           pExpr->op = TK_AGG_FUNCTION;
           pExpr->op2 = 0;
+#ifndef SQLITE_OMIT_WINDOWFUNC
+          if( ExprHasProperty(pExpr, EP_WinFunc) ){
+            sqlite3WalkExpr(pWalker, pExpr->y.pWin->pFilter);
+          }
+#endif
           while( pNC2 && !sqlite3FunctionUsesThisSrc(pExpr, pNC2->pSrcList) ){
             pExpr->op2++;
             pNC2 = pNC2->pNext;
           }
-          assert( pDef!=0 );
-          if( pNC2 ){
+          assert( pDef!=0 || IN_RENAME_OBJECT );
+          if( pNC2 && pDef ){
             assert( SQLITE_FUNC_MINMAX==NC_MinMaxAgg );
             testcase( (pDef->funcFlags & SQLITE_FUNC_MINMAX)!=0 );
             pNC2->ncFlags |= NC_HasAgg | (pDef->funcFlags & SQLITE_FUNC_MINMAX);
@@ -96023,7 +97786,12 @@ static int resolveExprStep(Walker *pWalker, Expr *pExpr){
       testcase( pExpr->op==TK_IN );
       if( ExprHasProperty(pExpr, EP_xIsSelect) ){
         int nRef = pNC->nRef;
-        notValid(pParse, pNC, "subqueries", NC_IsCheck|NC_PartIdx|NC_IdxExpr);
+        testcase( pNC->ncFlags & NC_IsCheck );
+        testcase( pNC->ncFlags & NC_PartIdx );
+        testcase( pNC->ncFlags & NC_IdxExpr );
+        testcase( pNC->ncFlags & NC_GenCol );
+        sqlite3ResolveNotValid(pParse, pNC, "subqueries",
+                 NC_IsCheck|NC_PartIdx|NC_IdxExpr|NC_GenCol, pExpr);
         sqlite3WalkSelect(pWalker, pExpr->x.pSelect);
         assert( pNC->nRef>=nRef );
         if( nRef!=pNC->nRef ){
@@ -96034,16 +97802,21 @@ static int resolveExprStep(Walker *pWalker, Expr *pExpr){
       break;
     }
     case TK_VARIABLE: {
-      notValid(pParse, pNC, "parameters", NC_IsCheck|NC_PartIdx|NC_IdxExpr);
+      testcase( pNC->ncFlags & NC_IsCheck );
+      testcase( pNC->ncFlags & NC_PartIdx );
+      testcase( pNC->ncFlags & NC_IdxExpr );
+      testcase( pNC->ncFlags & NC_GenCol );
+      sqlite3ResolveNotValid(pParse, pNC, "parameters",
+               NC_IsCheck|NC_PartIdx|NC_IdxExpr|NC_GenCol, pExpr);
       break;
     }
     case TK_IS:
     case TK_ISNOT: {
-      Expr *pRight;
+      Expr *pRight = sqlite3ExprSkipCollateAndLikely(pExpr->pRight);
       assert( !ExprHasProperty(pExpr, EP_Reduced) );
       /* Handle special cases of "x IS TRUE", "x IS FALSE", "x IS NOT TRUE",
       ** and "x IS NOT FALSE". */
-      if( (pRight = pExpr->pRight)->op==TK_ID ){
+      if( pRight->op==TK_ID ){
         int rc = resolveExprStep(pWalker, pRight);
         if( rc==WRC_Abort ) return WRC_Abort;
         if( pRight->op==TK_TRUEFALSE ){
@@ -96116,8 +97889,9 @@ static int resolveAsName(
   if( pE->op==TK_ID ){
     char *zCol = pE->u.zToken;
     for(i=0; i<pEList->nExpr; i++){
-      char *zAs = pEList->a[i].zName;
-      if( zAs!=0 && sqlite3StrICmp(zAs, zCol)==0 ){
+      if( pEList->a[i].eEName==ENAME_NAME
+       && sqlite3_stricmp(pEList->a[i].zEName, zCol)==0
+      ){
         return i+1;
       }
     }
@@ -96250,7 +98024,7 @@ static int resolveCompoundOrderBy(
       int iCol = -1;
       Expr *pE, *pDup;
       if( pItem->done ) continue;
-      pE = sqlite3ExprSkipCollate(pItem->pExpr);
+      pE = sqlite3ExprSkipCollateAndLikely(pItem->pExpr);
       if( sqlite3ExprIsInteger(pE, &iCol) ){
         if( iCol<=0 || iCol>pEList->nExpr ){
           resolveOutOfRangeError(pParse, "ORDER", i+1, pEList->nExpr);
@@ -96344,7 +98118,7 @@ SQLITE_PRIVATE int sqlite3ResolveOrderGroupBy(
   ExprList *pEList;
   struct ExprList_item *pItem;
 
-  if( pOrderBy==0 || pParse->db->mallocFailed ) return 0;
+  if( pOrderBy==0 || pParse->db->mallocFailed || IN_RENAME_OBJECT ) return 0;
   if( pOrderBy->nExpr>db->aLimit[SQLITE_LIMIT_COLUMN] ){
     sqlite3ErrorMsg(pParse, "too many terms in %s BY clause", zType);
     return 1;
@@ -96366,17 +98140,13 @@ SQLITE_PRIVATE int sqlite3ResolveOrderGroupBy(
 
 #ifndef SQLITE_OMIT_WINDOWFUNC
 /*
-** Walker callback for resolveRemoveWindows().
+** Walker callback for windowRemoveExprFromSelect().
 */
 static int resolveRemoveWindowsCb(Walker *pWalker, Expr *pExpr){
+  UNUSED_PARAMETER(pWalker);
   if( ExprHasProperty(pExpr, EP_WinFunc) ){
-    Window **pp;
-    for(pp=&pWalker->u.pSelect->pWin; *pp; pp=&(*pp)->pNextWin){
-      if( *pp==pExpr->y.pWin ){
-        *pp = (*pp)->pNextWin;
-        break;
-      }    
-    }
+    Window *pWin = pExpr->y.pWin;
+    sqlite3WindowUnlinkFromSelect(pWin);
   }
   return WRC_Continue;
 }
@@ -96385,16 +98155,18 @@ static int resolveRemoveWindowsCb(Walker *pWalker, Expr *pExpr){
 ** Remove any Window objects owned by the expression pExpr from the
 ** Select.pWin list of Select object pSelect.
 */
-static void resolveRemoveWindows(Select *pSelect, Expr *pExpr){
-  Walker sWalker;
-  memset(&sWalker, 0, sizeof(Walker));
-  sWalker.xExprCallback = resolveRemoveWindowsCb;
-  sWalker.u.pSelect = pSelect;
-  sqlite3WalkExpr(&sWalker, pExpr);
+static void windowRemoveExprFromSelect(Select *pSelect, Expr *pExpr){
+  if( pSelect->pWin ){
+    Walker sWalker;
+    memset(&sWalker, 0, sizeof(Walker));
+    sWalker.xExprCallback = resolveRemoveWindowsCb;
+    sWalker.u.pSelect = pSelect;
+    sqlite3WalkExpr(&sWalker, pExpr);
+  }
 }
 #else
-# define resolveRemoveWindows(x,y)
-#endif
+# define windowRemoveExprFromSelect(a, b)
+#endif /* SQLITE_OMIT_WINDOWFUNC */
 
 /*
 ** pOrderBy is an ORDER BY or GROUP BY clause in SELECT statement pSelect.
@@ -96431,7 +98203,7 @@ static int resolveOrderGroupBy(
   pParse = pNC->pParse;
   for(i=0, pItem=pOrderBy->a; i<pOrderBy->nExpr; i++, pItem++){
     Expr *pE = pItem->pExpr;
-    Expr *pE2 = sqlite3ExprSkipCollate(pE);
+    Expr *pE2 = sqlite3ExprSkipCollateAndLikely(pE);
     if( zType[0]!='G' ){
       iCol = resolveAsName(pParse, pSelect->pEList, pE2);
       if( iCol>0 ){
@@ -96465,7 +98237,7 @@ static int resolveOrderGroupBy(
         /* Since this expresion is being changed into a reference
         ** to an identical expression in the result set, remove all Window
         ** objects belonging to the expression from the Select.pWin list. */
-        resolveRemoveWindows(pSelect, pE);
+        windowRemoveExprFromSelect(pSelect, pE);
         pItem->u.x.iOrderByCol = j+1;
       }
     }
@@ -96765,7 +98537,7 @@ SQLITE_PRIVATE int sqlite3ResolveExprNames(
   NameContext *pNC,       /* Namespace to resolve expressions in. */
   Expr *pExpr             /* The expression to be analyzed. */
 ){
-  u16 savedHasAgg;
+  int savedHasAgg;
   Walker w;
 
   if( pExpr==0 ) return SQLITE_OK;
@@ -96845,10 +98617,13 @@ SQLITE_PRIVATE void sqlite3ResolveSelectNames(
 ** Resolve names in expressions that can only reference a single table
 ** or which cannot reference any tables at all.  Examples:
 **
-**    (1)   CHECK constraints
-**    (2)   WHERE clauses on partial indices
-**    (3)   Expressions in indexes on expressions
-**    (4)   Expression arguments to VACUUM INTO.
+**                                                    "type" flag
+**                                                    ------------
+**    (1)   CHECK constraints                         NC_IsCheck
+**    (2)   WHERE clauses on partial indices          NC_PartIdx
+**    (3)   Expressions in indexes on expressions     NC_IdxExpr
+**    (4)   Expression arguments to VACUUM INTO.      0
+**    (5)   GENERATED ALWAYS as expressions           NC_GenCol
 **
 ** In all cases except (4), the Expr.iTable value for Expr.op==TK_COLUMN
 ** nodes of the expression is set to -1 and the Expr.iColumn value is
@@ -96857,18 +98632,19 @@ SQLITE_PRIVATE void sqlite3ResolveSelectNames(
 ** Any errors cause an error message to be set in pParse.
 */
 SQLITE_PRIVATE int sqlite3ResolveSelfReference(
-  Parse *pParse,      /* Parsing context */
-  Table *pTab,        /* The table being referenced, or NULL */
-  int type,           /* NC_IsCheck or NC_PartIdx or NC_IdxExpr, or 0 */
-  Expr *pExpr,        /* Expression to resolve.  May be NULL. */
-  ExprList *pList     /* Expression list to resolve.  May be NULL. */
+  Parse *pParse,   /* Parsing context */
+  Table *pTab,     /* The table being referenced, or NULL */
+  int type,        /* NC_IsCheck, NC_PartIdx, NC_IdxExpr, NC_GenCol, or 0 */
+  Expr *pExpr,     /* Expression to resolve.  May be NULL. */
+  ExprList *pList  /* Expression list to resolve.  May be NULL. */
 ){
   SrcList sSrc;                   /* Fake SrcList for pParse->pNewTable */
   NameContext sNC;                /* Name context for pParse->pNewTable */
   int rc;
 
   assert( type==0 || pTab!=0 );
-  assert( type==NC_IsCheck || type==NC_PartIdx || type==NC_IdxExpr || pTab==0 );
+  assert( type==NC_IsCheck || type==NC_PartIdx || type==NC_IdxExpr
+          || type==NC_GenCol || pTab==0 );
   memset(&sNC, 0, sizeof(sNC));
   memset(&sSrc, 0, sizeof(sSrc));
   if( pTab ){
@@ -96876,10 +98652,15 @@ SQLITE_PRIVATE int sqlite3ResolveSelfReference(
     sSrc.a[0].zName = pTab->zName;
     sSrc.a[0].pTab = pTab;
     sSrc.a[0].iCursor = -1;
+    if( pTab->pSchema!=pParse->db->aDb[1].pSchema ){
+      /* Cause EP_FromDDL to be set on TK_FUNCTION nodes of non-TEMP
+      ** schema elements */
+      type |= NC_FromDDL;
+    }
   }
   sNC.pParse = pParse;
   sNC.pSrcList = &sSrc;
-  sNC.ncFlags = type;
+  sNC.ncFlags = type | NC_IsDDL;
   if( (rc = sqlite3ResolveExprNames(&sNC, pExpr))!=SQLITE_OK ) return rc;
   if( pList ) rc = sqlite3ResolveExprListNames(&sNC, pList);
   return rc;
@@ -96933,8 +98714,11 @@ SQLITE_PRIVATE char sqlite3TableColumnAffinity(Table *pTab, int iCol){
 */
 SQLITE_PRIVATE char sqlite3ExprAffinity(Expr *pExpr){
   int op;
-  pExpr = sqlite3ExprSkipCollate(pExpr);
-  if( pExpr->flags & EP_Generic ) return 0;
+  while( ExprHasProperty(pExpr, EP_Skip) ){
+    assert( pExpr->op==TK_COLLATE );
+    pExpr = pExpr->pLeft;
+    assert( pExpr!=0 );
+  }
   op = pExpr->op;
   if( op==TK_SELECT ){
     assert( pExpr->flags&EP_xIsSelect );
@@ -96956,7 +98740,10 @@ SQLITE_PRIVATE char sqlite3ExprAffinity(Expr *pExpr){
         pExpr->pLeft->x.pSelect->pEList->a[pExpr->iColumn].pExpr
     );
   }
-  return pExpr->affinity;
+  if( op==TK_VECTOR ){
+    return sqlite3ExprAffinity(pExpr->x.pList->a[0].pExpr);
+  }
+  return pExpr->affExpr;
 }
 
 /*
@@ -96991,11 +98778,23 @@ SQLITE_PRIVATE Expr *sqlite3ExprAddCollateString(Parse *pParse, Expr *pExpr, con
 }
 
 /*
-** Skip over any TK_COLLATE operators and any unlikely()
-** or likelihood() function at the root of an expression.
+** Skip over any TK_COLLATE operators.
 */
 SQLITE_PRIVATE Expr *sqlite3ExprSkipCollate(Expr *pExpr){
   while( pExpr && ExprHasProperty(pExpr, EP_Skip) ){
+    assert( pExpr->op==TK_COLLATE );
+    pExpr = pExpr->pLeft;
+  }   
+  return pExpr;
+}
+
+/*
+** Skip over any TK_COLLATE operators and/or any unlikely()
+** or likelihood() or likely() functions at the root of an
+** expression.
+*/
+SQLITE_PRIVATE Expr *sqlite3ExprSkipCollateAndLikely(Expr *pExpr){
+  while( pExpr && ExprHasProperty(pExpr, EP_Skip|EP_Unlikely) ){
     if( ExprHasProperty(pExpr, EP_Unlikely) ){
       assert( !ExprHasProperty(pExpr, EP_xIsSelect) );
       assert( pExpr->x.pList->nExpr>0 );
@@ -97029,7 +98828,6 @@ SQLITE_PRIVATE CollSeq *sqlite3ExprCollSeq(Parse *pParse, Expr *pExpr){
   Expr *p = pExpr;
   while( p ){
     int op = p->op;
-    if( p->flags & EP_Generic ) break;
     if( op==TK_REGISTER ) op = p->op2;
     if( (op==TK_AGG_COLUMN || op==TK_COLUMN || op==TK_TRIGGER)
      && p->y.pTab!=0
@@ -97047,6 +98845,10 @@ SQLITE_PRIVATE CollSeq *sqlite3ExprCollSeq(Parse *pParse, Expr *pExpr){
       p = p->pLeft;
       continue;
     }
+    if( op==TK_VECTOR ){
+      p = p->x.pList->a[0].pExpr;
+      continue;
+    }
     if( op==TK_COLLATE ){
       pColl = sqlite3GetCollSeq(pParse, ENC(db), 0, p->u.zToken);
       break;
@@ -97058,12 +98860,12 @@ SQLITE_PRIVATE CollSeq *sqlite3ExprCollSeq(Parse *pParse, Expr *pExpr){
         Expr *pNext  = p->pRight;
         /* The Expr.x union is never used at the same time as Expr.pRight */
         assert( p->x.pList==0 || p->pRight==0 );
-        /* p->flags holds EP_Collate and p->pLeft->flags does not.  And
-        ** p->x.pSelect cannot.  So if p->x.pLeft exists, it must hold at
-        ** least one EP_Collate. Thus the following two ALWAYS. */
-        if( p->x.pList!=0 && ALWAYS(!ExprHasProperty(p, EP_xIsSelect)) ){
+        if( p->x.pList!=0 
+         && !db->mallocFailed
+         && ALWAYS(!ExprHasProperty(p, EP_xIsSelect))
+        ){
           int i;
-          for(i=0; ALWAYS(i<p->x.pList->nExpr); i++){
+          for(i=0; i<p->x.pList->nExpr; i++){
             if( ExprHasProperty(p->x.pList->a[i].pExpr, EP_Collate) ){
               pNext = p->x.pList->a[i].pExpr;
               break;
@@ -97115,7 +98917,7 @@ SQLITE_PRIVATE int sqlite3ExprCollSeqMatch(Parse *pParse, Expr *pE1, Expr *pE2){
 */
 SQLITE_PRIVATE char sqlite3CompareAffinity(Expr *pExpr, char aff2){
   char aff1 = sqlite3ExprAffinity(pExpr);
-  if( aff1 && aff2 ){
+  if( aff1>SQLITE_AFF_NONE && aff2>SQLITE_AFF_NONE ){
     /* Both sides of the comparison are columns. If one has numeric
     ** affinity, use that. Otherwise use no affinity.
     */
@@ -97124,15 +98926,10 @@ SQLITE_PRIVATE char sqlite3CompareAffinity(Expr *pExpr, char aff2){
     }else{
       return SQLITE_AFF_BLOB;
     }
-  }else if( !aff1 && !aff2 ){
-    /* Neither side of the comparison is a column.  Compare the
-    ** results directly.
-    */
-    return SQLITE_AFF_BLOB;
   }else{
     /* One side is a column, the other is not. Use the columns affinity. */
-    assert( aff1==0 || aff2==0 );
-    return (aff1 + aff2);
+    assert( aff1<=SQLITE_AFF_NONE || aff2<=SQLITE_AFF_NONE );
+    return (aff1<=SQLITE_AFF_NONE ? aff2 : aff1) | SQLITE_AFF_NONE;
   }
 }
 
@@ -97165,14 +98962,13 @@ static char comparisonAffinity(Expr *pExpr){
 */
 SQLITE_PRIVATE int sqlite3IndexAffinityOk(Expr *pExpr, char idx_affinity){
   char aff = comparisonAffinity(pExpr);
-  switch( aff ){
-    case SQLITE_AFF_BLOB:
-      return 1;
-    case SQLITE_AFF_TEXT:
-      return idx_affinity==SQLITE_AFF_TEXT;
-    default:
-      return sqlite3IsNumericAffinity(idx_affinity);
+  if( aff<SQLITE_AFF_TEXT ){
+    return 1;
   }
+  if( aff==SQLITE_AFF_TEXT ){
+    return idx_affinity==SQLITE_AFF_TEXT;
+  }
+  return sqlite3IsNumericAffinity(idx_affinity);
 }
 
 /*
@@ -97217,6 +99013,22 @@ SQLITE_PRIVATE CollSeq *sqlite3BinaryCompareCollSeq(
   return pColl;
 }
 
+/* Expresssion p is a comparison operator.  Return a collation sequence
+** appropriate for the comparison operator.
+**
+** This is normally just a wrapper around sqlite3BinaryCompareCollSeq().
+** However, if the OP_Commuted flag is set, then the order of the operands
+** is reversed in the sqlite3BinaryCompareCollSeq() call so that the
+** correct collating sequence is found.
+*/
+SQLITE_PRIVATE CollSeq *sqlite3ExprCompareCollSeq(Parse *pParse, Expr *p){
+  if( ExprHasProperty(p, EP_Commuted) ){
+    return sqlite3BinaryCompareCollSeq(pParse, p->pRight, p->pLeft);
+  }else{
+    return sqlite3BinaryCompareCollSeq(pParse, p->pLeft, p->pRight);
+  }
+}
+
 /*
 ** Generate code for a comparison operator.
 */
@@ -97227,13 +99039,19 @@ static int codeCompare(
   int opcode,       /* The comparison opcode */
   int in1, int in2, /* Register holding operands */
   int dest,         /* Jump here if true.  */
-  int jumpIfNull    /* If true, jump if either operand is NULL */
+  int jumpIfNull,   /* If true, jump if either operand is NULL */
+  int isCommuted    /* The comparison has been commuted */
 ){
   int p5;
   int addr;
   CollSeq *p4;
 
-  p4 = sqlite3BinaryCompareCollSeq(pParse, pLeft, pRight);
+  if( pParse->nErr ) return 0;
+  if( isCommuted ){
+    p4 = sqlite3BinaryCompareCollSeq(pParse, pRight, pLeft);
+  }else{
+    p4 = sqlite3BinaryCompareCollSeq(pParse, pLeft, pRight);
+  }
   p5 = binaryCompareP5(pLeft, pRight, jumpIfNull);
   addr = sqlite3VdbeAddOp4(pParse->pVdbe, opcode, in2, dest, in1,
                            (void*)p4, P4_COLLSEQ);
@@ -97444,7 +99262,9 @@ static void codeVectorCompare(
   int regRight = 0;
   u8 opx = op;
   int addrDone = sqlite3VdbeMakeLabel(pParse);
+  int isCommuted = ExprHasProperty(pExpr,EP_Commuted);
 
+  if( pParse->nErr ) return;
   if( nLeft!=sqlite3ExprVectorSize(pRight) ){
     sqlite3ErrorMsg(pParse, "row value misused");
     return;
@@ -97473,7 +99293,7 @@ static void codeVectorCompare(
     assert( i>=0 && i<nLeft );
     r1 = exprVectorRegister(pParse, pLeft, i, regLeft, &pL, &regFree1);
     r2 = exprVectorRegister(pParse, pRight, i, regRight, &pR, &regFree2);
-    codeCompare(pParse, pL, pR, opx, r1, r2, dest, p5);
+    codeCompare(pParse, pL, pR, opx, r1, r2, dest, p5, isCommuted);
     testcase(op==OP_Lt); VdbeCoverageIf(v,op==OP_Lt);
     testcase(op==OP_Le); VdbeCoverageIf(v,op==OP_Le);
     testcase(op==OP_Gt); VdbeCoverageIf(v,op==OP_Gt);
@@ -97662,7 +99482,7 @@ SQLITE_PRIVATE Expr *sqlite3ExprAlloc(
     pNew->iAgg = -1;
     if( pToken ){
       if( nExtra==0 ){
-        pNew->flags |= EP_IntValue|EP_Leaf;
+        pNew->flags |= EP_IntValue|EP_Leaf|(iValue?EP_IsTrue:EP_IsFalse);
         pNew->u.iValue = iValue;
       }else{
         pNew->u.zToken = (char*)&pNew[1];
@@ -97739,20 +99559,16 @@ SQLITE_PRIVATE Expr *sqlite3PExpr(
   Expr *pRight            /* Right operand */
 ){
   Expr *p;
-  if( op==TK_AND && pParse->nErr==0 && !IN_RENAME_OBJECT ){
-    /* Take advantage of short-circuit false optimization for AND */
-    p = sqlite3ExprAnd(pParse->db, pLeft, pRight);
-  }else{
-    p = sqlite3DbMallocRawNN(pParse->db, sizeof(Expr));
-    if( p ){
-      memset(p, 0, sizeof(Expr));
-      p->op = op & 0xff;
-      p->iAgg = -1;
-    }
+  p = sqlite3DbMallocRawNN(pParse->db, sizeof(Expr));
+  if( p ){
+    memset(p, 0, sizeof(Expr));
+    p->op = op & 0xff;
+    p->iAgg = -1;
     sqlite3ExprAttachSubtrees(pParse->db, p, pLeft, pRight);
-  }
-  if( p ) {
     sqlite3ExprCheckHeight(pParse, p->nHeight);
+  }else{
+    sqlite3ExprDelete(pParse->db, pLeft);
+    sqlite3ExprDelete(pParse->db, pRight);
   }
   return p;
 }
@@ -97773,33 +99589,6 @@ SQLITE_PRIVATE void sqlite3PExprAddSelect(Parse *pParse, Expr *pExpr, Select *pS
 }
 
 
-/*
-** If the expression is always either TRUE or FALSE (respectively),
-** then return 1.  If one cannot determine the truth value of the
-** expression at compile-time return 0.
-**
-** This is an optimization.  If is OK to return 0 here even if
-** the expression really is always false or false (a false negative).
-** But it is a bug to return 1 if the expression might have different
-** boolean values in different circumstances (a false positive.)
-**
-** Note that if the expression is part of conditional for a
-** LEFT JOIN, then we cannot determine at compile-time whether or not
-** is it true or false, so always return 0.
-*/
-static int exprAlwaysTrue(Expr *p){
-  int v = 0;
-  if( ExprHasProperty(p, EP_FromJoin) ) return 0;
-  if( !sqlite3ExprIsInteger(p, &v) ) return 0;
-  return v!=0;
-}
-static int exprAlwaysFalse(Expr *p){
-  int v = 0;
-  if( ExprHasProperty(p, EP_FromJoin) ) return 0;
-  if( !sqlite3ExprIsInteger(p, &v) ) return 0;
-  return v==0;
-}
-
 /*
 ** Join two expressions using an AND operator.  If either expression is
 ** NULL, then just return the other expression.
@@ -97808,19 +99597,20 @@ static int exprAlwaysFalse(Expr *p){
 ** of returning an AND expression, just return a constant expression with
 ** a value of false.
 */
-SQLITE_PRIVATE Expr *sqlite3ExprAnd(sqlite3 *db, Expr *pLeft, Expr *pRight){
-  if( pLeft==0 ){
+SQLITE_PRIVATE Expr *sqlite3ExprAnd(Parse *pParse, Expr *pLeft, Expr *pRight){
+  sqlite3 *db = pParse->db;
+  if( pLeft==0  ){
     return pRight;
   }else if( pRight==0 ){
     return pLeft;
-  }else if( exprAlwaysFalse(pLeft) || exprAlwaysFalse(pRight) ){
+  }else if( (ExprAlwaysFalse(pLeft) || ExprAlwaysFalse(pRight)) 
+         && !IN_RENAME_OBJECT
+  ){
     sqlite3ExprDelete(db, pLeft);
     sqlite3ExprDelete(db, pRight);
-    return sqlite3ExprAlloc(db, TK_INTEGER, &sqlite3IntTokens[0], 0);
+    return sqlite3Expr(db, TK_INTEGER, "0");
   }else{
-    Expr *pNew = sqlite3ExprAlloc(db, TK_AND, 0, 0);
-    sqlite3ExprAttachSubtrees(db, pNew, pLeft, pRight);
-    return pNew;
+    return sqlite3PExpr(pParse, TK_AND, pLeft, pRight);
   }
 }
 
@@ -97853,6 +99643,40 @@ SQLITE_PRIVATE Expr *sqlite3ExprFunction(
   return pNew;
 }
 
+/*
+** Check to see if a function is usable according to current access
+** rules:
+**
+**    SQLITE_FUNC_DIRECT    -     Only usable from top-level SQL
+**
+**    SQLITE_FUNC_UNSAFE    -     Usable if TRUSTED_SCHEMA or from
+**                                top-level SQL
+**
+** If the function is not usable, create an error.
+*/
+SQLITE_PRIVATE void sqlite3ExprFunctionUsable(
+  Parse *pParse,         /* Parsing and code generating context */
+  Expr *pExpr,           /* The function invocation */
+  FuncDef *pDef          /* The function being invoked */
+){
+  assert( !IN_RENAME_OBJECT );
+  assert( (pDef->funcFlags & (SQLITE_FUNC_DIRECT|SQLITE_FUNC_UNSAFE))!=0 );
+  if( ExprHasProperty(pExpr, EP_FromDDL) ){
+    if( (pDef->funcFlags & SQLITE_FUNC_DIRECT)!=0
+     || (pParse->db->flags & SQLITE_TrustedSchema)==0
+    ){
+      /* Functions prohibited in triggers and views if:
+      **     (1) tagged with SQLITE_DIRECTONLY
+      **     (2) not tagged with SQLITE_INNOCUOUS (which means it
+      **         is tagged with SQLITE_FUNC_UNSAFE) and 
+      **         SQLITE_DBCONFIG_TRUSTED_SCHEMA is off (meaning
+      **         that the schema is possibly tainted).
+      */
+      sqlite3ErrorMsg(pParse, "unsafe use of %s()", pDef->zName);
+    }
+  }
+}
+
 /*
 ** Assign a variable number to an expression that encodes a wildcard
 ** in the original SQL statement.  
@@ -97957,15 +99781,18 @@ static SQLITE_NOINLINE void sqlite3ExprDeleteNN(sqlite3 *db, Expr *p){
     assert( p->x.pList==0 || p->pRight==0 );
     if( p->pLeft && p->op!=TK_SELECT_COLUMN ) sqlite3ExprDeleteNN(db, p->pLeft);
     if( p->pRight ){
+      assert( !ExprHasProperty(p, EP_WinFunc) );
       sqlite3ExprDeleteNN(db, p->pRight);
     }else if( ExprHasProperty(p, EP_xIsSelect) ){
+      assert( !ExprHasProperty(p, EP_WinFunc) );
       sqlite3SelectDelete(db, p->x.pSelect);
     }else{
       sqlite3ExprListDelete(db, p->x.pList);
-    }
-    if( ExprHasProperty(p, EP_WinFunc) ){
-      assert( p->op==TK_FUNCTION );
-      sqlite3WindowDelete(db, p->y.pWin);
+#ifndef SQLITE_OMIT_WINDOWFUNC
+      if( ExprHasProperty(p, EP_WinFunc) ){
+        sqlite3WindowDelete(db, p->y.pWin);
+      }
+#endif
     }
   }
   if( ExprHasProperty(p, EP_MemToken) ) sqlite3DbFree(db, p->u.zToken);
@@ -97977,6 +99804,18 @@ SQLITE_PRIVATE void sqlite3ExprDelete(sqlite3 *db, Expr *p){
   if( p ) sqlite3ExprDeleteNN(db, p);
 }
 
+/* Invoke sqlite3RenameExprUnmap() and sqlite3ExprDelete() on the
+** expression.
+*/
+SQLITE_PRIVATE void sqlite3ExprUnmapAndDelete(Parse *pParse, Expr *p){
+  if( p ){
+    if( IN_RENAME_OBJECT ){
+      sqlite3RenameExprUnmap(pParse, p);
+    }
+    sqlite3ExprDeleteNN(pParse->db, p);
+  }
+}
+
 /*
 ** Return the number of bytes allocated for the expression structure 
 ** passed as the first argument. This is always one of EXPR_FULLSIZE,
@@ -97988,16 +99827,6 @@ static int exprStructSize(Expr *p){
   return EXPR_FULLSIZE;
 }
 
-/*
-** Copy the complete content of an Expr node, taking care not to read
-** past the end of the structure for a reduced-size version of the source
-** Expr.
-*/
-static void exprNodeCopy(Expr *pDest, Expr *pSrc){
-  memset(pDest, 0, sizeof(Expr));
-  memcpy(pDest, pSrc, exprStructSize(pSrc));
-}
-
 /*
 ** The dupedExpr*Size() routines each return the number of bytes required
 ** to store a copy of an expression or expression tree.  They differ in
@@ -98237,10 +100066,13 @@ static With *withDup(sqlite3 *db, With *p){
 ** objects found there, assembling them onto the linked list at Select->pWin.
 */
 static int gatherSelectWindowsCallback(Walker *pWalker, Expr *pExpr){
-  if( pExpr->op==TK_FUNCTION && pExpr->y.pWin!=0 ){
-    assert( ExprHasProperty(pExpr, EP_WinFunc) );
-    pExpr->y.pWin->pNextWin = pWalker->u.pSelect->pWin;
-    pWalker->u.pSelect->pWin = pExpr->y.pWin;
+  if( pExpr->op==TK_FUNCTION && ExprHasProperty(pExpr, EP_WinFunc) ){
+    Select *pSelect = pWalker->u.pSelect;
+    Window *pWin = pExpr->y.pWin;
+    assert( pWin );
+    assert( IsWindowFunc(pExpr) );
+    assert( pWin->ppThis==0 );
+    sqlite3WindowLink(pSelect, pWin);
   }
   return WRC_Continue;
 }
@@ -98312,11 +100144,11 @@ SQLITE_PRIVATE ExprList *sqlite3ExprListDup(sqlite3 *db, ExprList *p, int flags)
         pNewExpr->pLeft = pPriorSelectCol;
       }
     }
-    pItem->zName = sqlite3DbStrDup(db, pOldItem->zName);
-    pItem->zSpan = sqlite3DbStrDup(db, pOldItem->zSpan);
-    pItem->sortOrder = pOldItem->sortOrder;
+    pItem->zEName = sqlite3DbStrDup(db, pOldItem->zEName);
+    pItem->sortFlags = pOldItem->sortFlags;
+    pItem->eEName = pOldItem->eEName;
     pItem->done = 0;
-    pItem->bSpanIsTab = pOldItem->bSpanIsTab;
+    pItem->bNulls = pOldItem->bNulls;
     pItem->bSorterRef = pOldItem->bSorterRef;
     pItem->u = pOldItem->u;
   }
@@ -98426,7 +100258,7 @@ SQLITE_PRIVATE Select *sqlite3SelectDup(sqlite3 *db, Select *pDup, int flags){
 #ifndef SQLITE_OMIT_WINDOWFUNC
     pNew->pWin = 0;
     pNew->pWinDefn = sqlite3WindowListDup(db, p->pWinDefn);
-    if( p->pWin ) gatherSelectWindows(pNew);
+    if( p->pWin && db->mallocFailed==0 ) gatherSelectWindows(pNew);
 #endif
     pNew->selId = p->selId;
     *pp = pNew;
@@ -98483,9 +100315,9 @@ SQLITE_PRIVATE ExprList *sqlite3ExprListAppend(
     pList = pNew;
   }
   pItem = &pList->a[pList->nExpr++];
-  assert( offsetof(struct ExprList_item,zName)==sizeof(pItem->pExpr) );
+  assert( offsetof(struct ExprList_item,zEName)==sizeof(pItem->pExpr) );
   assert( offsetof(struct ExprList_item,pExpr)==0 );
-  memset(&pItem->zName,0,sizeof(*pItem)-offsetof(struct ExprList_item,zName));
+  memset(&pItem->zEName,0,sizeof(*pItem)-offsetof(struct ExprList_item,zEName));
   pItem->pExpr = pExpr;
   return pList;
 
@@ -98535,10 +100367,14 @@ SQLITE_PRIVATE ExprList *sqlite3ExprListAppendVector(
 
   for(i=0; i<pColumns->nId; i++){
     Expr *pSubExpr = sqlite3ExprForVectorField(pParse, pExpr, i);
+    assert( pSubExpr!=0 || db->mallocFailed );
+    assert( pSubExpr==0 || pSubExpr->iTable==0 );
+    if( pSubExpr==0 ) continue;
+    pSubExpr->iTable = pColumns->nId;
     pList = sqlite3ExprListAppend(pParse, pList, pSubExpr);
     if( pList ){
       assert( pList->nExpr==iFirst+i+1 );
-      pList->a[pList->nExpr-1].zName = pColumns->a[i].zName;
+      pList->a[pList->nExpr-1].zEName = pColumns->a[i].zName;
       pColumns->a[i].zName = 0;
     }
   }
@@ -98559,10 +100395,7 @@ SQLITE_PRIVATE ExprList *sqlite3ExprListAppendVector(
   }
 
 vector_append_error:
-  if( IN_RENAME_OBJECT ){
-    sqlite3RenameExprUnmap(pParse, pExpr);
-  }
-  sqlite3ExprDelete(db, pExpr);
+  sqlite3ExprUnmapAndDelete(pParse, pExpr);
   sqlite3IdListDelete(db, pColumns);
   return pList;
 }
@@ -98570,19 +100403,38 @@ vector_append_error:
 /*
 ** Set the sort order for the last element on the given ExprList.
 */
-SQLITE_PRIVATE void sqlite3ExprListSetSortOrder(ExprList *p, int iSortOrder){
+SQLITE_PRIVATE void sqlite3ExprListSetSortOrder(ExprList *p, int iSortOrder, int eNulls){
+  struct ExprList_item *pItem;
   if( p==0 ) return;
-  assert( SQLITE_SO_UNDEFINED<0 && SQLITE_SO_ASC>=0 && SQLITE_SO_DESC>0 );
   assert( p->nExpr>0 );
-  if( iSortOrder<0 ){
-    assert( p->a[p->nExpr-1].sortOrder==SQLITE_SO_ASC );
-    return;
+
+  assert( SQLITE_SO_UNDEFINED<0 && SQLITE_SO_ASC==0 && SQLITE_SO_DESC>0 );
+  assert( iSortOrder==SQLITE_SO_UNDEFINED 
+       || iSortOrder==SQLITE_SO_ASC 
+       || iSortOrder==SQLITE_SO_DESC 
+  );
+  assert( eNulls==SQLITE_SO_UNDEFINED 
+       || eNulls==SQLITE_SO_ASC 
+       || eNulls==SQLITE_SO_DESC 
+  );
+
+  pItem = &p->a[p->nExpr-1];
+  assert( pItem->bNulls==0 );
+  if( iSortOrder==SQLITE_SO_UNDEFINED ){
+    iSortOrder = SQLITE_SO_ASC;
+  }
+  pItem->sortFlags = (u8)iSortOrder;
+
+  if( eNulls!=SQLITE_SO_UNDEFINED ){
+    pItem->bNulls = 1;
+    if( iSortOrder!=eNulls ){
+      pItem->sortFlags |= KEYINFO_ORDER_BIGNULL;
+    }
   }
-  p->a[p->nExpr-1].sortOrder = (u8)iSortOrder;
 }
 
 /*
-** Set the ExprList.a[].zName element of the most recently added item
+** Set the ExprList.a[].zEName element of the most recently added item
 ** on the expression list.
 **
 ** pList might be NULL following an OOM error.  But pName should never be
@@ -98600,11 +100452,12 @@ SQLITE_PRIVATE void sqlite3ExprListSetName(
     struct ExprList_item *pItem;
     assert( pList->nExpr>0 );
     pItem = &pList->a[pList->nExpr-1];
-    assert( pItem->zName==0 );
-    pItem->zName = sqlite3DbStrNDup(pParse->db, pName->z, pName->n);
-    if( dequote ) sqlite3Dequote(pItem->zName);
+    assert( pItem->zEName==0 );
+    assert( pItem->eEName==ENAME_NAME );
+    pItem->zEName = sqlite3DbStrNDup(pParse->db, pName->z, pName->n);
+    if( dequote ) sqlite3Dequote(pItem->zEName);
     if( IN_RENAME_OBJECT ){
-      sqlite3RenameTokenMap(pParse, (void*)pItem->zName, pName);
+      sqlite3RenameTokenMap(pParse, (void*)pItem->zEName, pName);
     }
   }
 }
@@ -98628,8 +100481,10 @@ SQLITE_PRIVATE void sqlite3ExprListSetSpan(
   if( pList ){
     struct ExprList_item *pItem = &pList->a[pList->nExpr-1];
     assert( pList->nExpr>0 );
-    sqlite3DbFree(db, pItem->zSpan);
-    pItem->zSpan = sqlite3DbSpanDup(db, zStart, zEnd);
+    if( pItem->zEName==0 ){
+      pItem->zEName = sqlite3DbSpanDup(db, zStart, zEnd);
+      pItem->eEName = ENAME_SPAN;
+    }
   }
 }
 
@@ -98659,8 +100514,7 @@ static SQLITE_NOINLINE void exprListDeleteNN(sqlite3 *db, ExprList *pList){
   assert( pList->nExpr>0 );
   do{
     sqlite3ExprDelete(db, pItem->pExpr);
-    sqlite3DbFree(db, pItem->zName);
-    sqlite3DbFree(db, pItem->zSpan);
+    sqlite3DbFree(db, pItem->zEName);
     pItem++;
   }while( --i>0 );
   sqlite3DbFreeNN(db, pList);
@@ -98698,18 +100552,34 @@ SQLITE_PRIVATE int sqlite3SelectWalkFail(Walker *pWalker, Select *NotUsed){
   return WRC_Abort;
 }
 
+/*
+** Check the input string to see if it is "true" or "false" (in any case).
+**
+**       If the string is....           Return
+**         "true"                         EP_IsTrue
+**         "false"                        EP_IsFalse
+**         anything else                  0
+*/
+SQLITE_PRIVATE u32 sqlite3IsTrueOrFalse(const char *zIn){
+  if( sqlite3StrICmp(zIn, "true")==0  ) return EP_IsTrue;
+  if( sqlite3StrICmp(zIn, "false")==0 ) return EP_IsFalse;
+  return 0;
+}
+
+
 /*
 ** If the input expression is an ID with the name "true" or "false"
 ** then convert it into an TK_TRUEFALSE term.  Return non-zero if
 ** the conversion happened, and zero if the expression is unaltered.
 */
 SQLITE_PRIVATE int sqlite3ExprIdToTrueFalse(Expr *pExpr){
+  u32 v;
   assert( pExpr->op==TK_ID || pExpr->op==TK_STRING );
   if( !ExprHasProperty(pExpr, EP_Quoted)
-   && (sqlite3StrICmp(pExpr->u.zToken, "true")==0
-       || sqlite3StrICmp(pExpr->u.zToken, "false")==0)
+   && (v = sqlite3IsTrueOrFalse(pExpr->u.zToken))!=0
   ){
     pExpr->op = TK_TRUEFALSE;
+    ExprSetProperty(pExpr, v);
     return 1;
   }
   return 0;
@@ -98720,12 +100590,40 @@ SQLITE_PRIVATE int sqlite3ExprIdToTrueFalse(Expr *pExpr){
 ** and 0 if it is FALSE.
 */
 SQLITE_PRIVATE int sqlite3ExprTruthValue(const Expr *pExpr){
+  pExpr = sqlite3ExprSkipCollate((Expr*)pExpr);
   assert( pExpr->op==TK_TRUEFALSE );
   assert( sqlite3StrICmp(pExpr->u.zToken,"true")==0
        || sqlite3StrICmp(pExpr->u.zToken,"false")==0 );
   return pExpr->u.zToken[4]==0;
 }
 
+/*
+** If pExpr is an AND or OR expression, try to simplify it by eliminating
+** terms that are always true or false.  Return the simplified expression.
+** Or return the original expression if no simplification is possible.
+**
+** Examples:
+**
+**     (x<10) AND true                =>   (x<10)
+**     (x<10) AND false               =>   false
+**     (x<10) AND (y=22 OR false)     =>   (x<10) AND (y=22)
+**     (x<10) AND (y=22 OR true)      =>   (x<10)
+**     (y=22) OR true                 =>   true
+*/
+SQLITE_PRIVATE Expr *sqlite3ExprSimplifiedAndOr(Expr *pExpr){
+  assert( pExpr!=0 );
+  if( pExpr->op==TK_AND || pExpr->op==TK_OR ){
+    Expr *pRight = sqlite3ExprSimplifiedAndOr(pExpr->pRight);
+    Expr *pLeft = sqlite3ExprSimplifiedAndOr(pExpr->pLeft);
+    if( ExprAlwaysTrue(pLeft) || ExprAlwaysFalse(pRight) ){
+      pExpr = pExpr->op==TK_AND ? pRight : pLeft;
+    }else if( ExprAlwaysTrue(pRight) || ExprAlwaysFalse(pLeft) ){
+      pExpr = pExpr->op==TK_AND ? pLeft : pRight;
+    }
+  }
+  return pExpr;
+}
+
 
 /*
 ** These routines are Walker callbacks used to check expressions to
@@ -98743,10 +100641,11 @@ SQLITE_PRIVATE int sqlite3ExprTruthValue(const Expr *pExpr){
 ** In all cases, the callbacks set Walker.eCode=0 and abort if the expression
 ** is found to not be a constant.
 **
-** The sqlite3ExprIsConstantOrFunction() is used for evaluating expressions
-** in a CREATE TABLE statement.  The Walker.eCode value is 5 when parsing
-** an existing schema and 4 when processing a new statement.  A bound
-** parameter raises an error for new statements, but is silently converted
+** The sqlite3ExprIsConstantOrFunction() is used for evaluating DEFAULT
+** expressions in a CREATE TABLE statement.  The Walker.eCode value is 5
+** when parsing an existing schema out of the sqlite_master table and 4
+** when processing a new CREATE TABLE statement.  A bound parameter raises
+** an error for new statements, but is silently converted
 ** to NULL for existing schemas.  This allows sqlite_master tables that 
 ** contain a bound parameter because they were generated by older versions
 ** of SQLite to be parsed by newer versions of SQLite without raising a
@@ -98767,7 +100666,10 @@ static int exprNodeIsConstant(Walker *pWalker, Expr *pExpr){
     ** and either pWalker->eCode==4 or 5 or the function has the
     ** SQLITE_FUNC_CONST flag. */
     case TK_FUNCTION:
-      if( pWalker->eCode>=4 || ExprHasProperty(pExpr,EP_ConstFunc) ){
+      if( (pWalker->eCode>=4 || ExprHasProperty(pExpr,EP_ConstFunc))
+       && !ExprHasProperty(pExpr, EP_WinFunc)
+      ){
+        if( pWalker->eCode==5 ) ExprSetProperty(pExpr, EP_FromDDL);
         return WRC_Continue;
       }else{
         pWalker->eCode = 0;
@@ -98931,9 +100833,21 @@ SQLITE_PRIVATE int sqlite3ExprIsConstantOrGroupBy(Parse *pParse, Expr *p, ExprLi
 }
 
 /*
-** Walk an expression tree.  Return non-zero if the expression is constant
-** or a function call with constant arguments.  Return and 0 if there
-** are any variables.
+** Walk an expression tree for the DEFAULT field of a column definition
+** in a CREATE TABLE statement.  Return non-zero if the expression is 
+** acceptable for use as a DEFAULT.  That is to say, return non-zero if
+** the expression is constant or a function call with constant arguments.
+** Return and 0 if there are any variables.
+**
+** isInit is true when parsing from sqlite_master.  isInit is false when
+** processing a new CREATE TABLE statement.  When isInit is true, parameters
+** (such as ? or $abc) in the expression are converted into NULL.  When
+** isInit is false, parameters raise an error.  Parameters should not be
+** allowed in a CREATE TABLE statement, but some legacy versions of SQLite
+** allowed it, so we need to support it when reading sqlite_master for
+** backwards compatibility.
+**
+** If isInit is true, set EP_FromDDL on every TK_FUNCTION node.
 **
 ** For the purposes of this function, a double-quoted string (ex: "abc")
 ** is considered a variable but a single-quoted string (ex: 'abc') is
@@ -98970,7 +100884,7 @@ SQLITE_PRIVATE int sqlite3ExprContainsSubquery(Expr *p){
 */
 SQLITE_PRIVATE int sqlite3ExprIsInteger(Expr *p, int *pValue){
   int rc = 0;
-  if( p==0 ) return 0;  /* Can only happen following on OOM */
+  if( NEVER(p==0) ) return 0;  /* Used to only happen following on OOM */
 
   /* If an expression is an integer literal that fits in a signed 32-bit
   ** integer, then the EP_IntValue flag will have already been set */
@@ -99030,7 +100944,9 @@ SQLITE_PRIVATE int sqlite3ExprCanBeNull(const Expr *p){
     case TK_COLUMN:
       return ExprHasProperty(p, EP_CanBeNull) ||
              p->y.pTab==0 ||  /* Reference to column of index on expression */
-             (p->iColumn>=0 && p->y.pTab->aCol[p->iColumn].notNull==0);
+             (p->iColumn>=0
+              && ALWAYS(p->y.pTab->aCol!=0) /* Defense against OOM problems */
+              && p->y.pTab->aCol[p->iColumn].notNull==0);
     default:
       return 1;
   }
@@ -99048,27 +100964,30 @@ SQLITE_PRIVATE int sqlite3ExprCanBeNull(const Expr *p){
 */
 SQLITE_PRIVATE int sqlite3ExprNeedsNoAffinityChange(const Expr *p, char aff){
   u8 op;
+  int unaryMinus = 0;
   if( aff==SQLITE_AFF_BLOB ) return 1;
-  while( p->op==TK_UPLUS || p->op==TK_UMINUS ){ p = p->pLeft; }
+  while( p->op==TK_UPLUS || p->op==TK_UMINUS ){
+    if( p->op==TK_UMINUS ) unaryMinus = 1;
+    p = p->pLeft;
+  }
   op = p->op;
   if( op==TK_REGISTER ) op = p->op2;
   switch( op ){
     case TK_INTEGER: {
-      return aff==SQLITE_AFF_INTEGER || aff==SQLITE_AFF_NUMERIC;
+      return aff>=SQLITE_AFF_NUMERIC;
     }
     case TK_FLOAT: {
-      return aff==SQLITE_AFF_REAL || aff==SQLITE_AFF_NUMERIC;
+      return aff>=SQLITE_AFF_NUMERIC;
     }
     case TK_STRING: {
-      return aff==SQLITE_AFF_TEXT;
+      return !unaryMinus && aff==SQLITE_AFF_TEXT;
     }
     case TK_BLOB: {
-      return 1;
+      return !unaryMinus;
     }
     case TK_COLUMN: {
       assert( p->iTable>=0 );  /* p cannot be part of a CHECK constraint */
-      return p->iColumn<0
-          && (aff==SQLITE_AFF_INTEGER || aff==SQLITE_AFF_NUMERIC);
+      return aff>=SQLITE_AFF_NUMERIC && p->iColumn<0;
     }
     default: {
       return 0;
@@ -99251,7 +101170,7 @@ static int sqlite3InRhsIsConstant(Expr *pIn){
 #ifndef SQLITE_OMIT_SUBQUERY
 SQLITE_PRIVATE int sqlite3FindInIndex(
   Parse *pParse,             /* Parsing context */
-  Expr *pX,                  /* The right-hand side (RHS) of the IN operator */
+  Expr *pX,                  /* The IN expression */
   u32 inFlags,               /* IN_INDEX_LOOP, _MEMBERSHIP, and/or _NOOP_OK */
   int *prRhsHasNull,         /* Register holding NULL status.  See notes */
   int *aiMap,                /* Mapping from Index fields to RHS fields */
@@ -99504,8 +101423,10 @@ static char *exprINAffinity(Parse *pParse, Expr *pExpr){
 **   "sub-select returns N columns - expected M"
 */   
 SQLITE_PRIVATE void sqlite3SubselectError(Parse *pParse, int nActual, int nExpect){
-  const char *zFmt = "sub-select returns %d columns - expected %d";
-  sqlite3ErrorMsg(pParse, zFmt, nActual, nExpect);
+  if( pParse->nErr==0 ){
+    const char *zFmt = "sub-select returns %d columns - expected %d";
+    sqlite3ErrorMsg(pParse, zFmt, nActual, nExpect);
+  }
 }
 #endif
 
@@ -99676,9 +101597,9 @@ SQLITE_PRIVATE void sqlite3CodeRhsOfIN(
     int i;
     ExprList *pList = pExpr->x.pList;
     struct ExprList_item *pItem;
-    int r1, r2, r3;
+    int r1, r2;
     affinity = sqlite3ExprAffinity(pLeft);
-    if( !affinity ){
+    if( affinity<=SQLITE_AFF_NONE ){
       affinity = SQLITE_AFF_BLOB;
     }
     if( pKeyInfo ){
@@ -99699,13 +101620,14 @@ SQLITE_PRIVATE void sqlite3CodeRhsOfIN(
       */
       if( addrOnce && !sqlite3ExprIsConstant(pE2) ){
         sqlite3VdbeChangeToNoop(v, addrOnce);
+        ExprClearProperty(pExpr, EP_Subrtn);
         addrOnce = 0;
       }
 
       /* Evaluate the expression and insert it into the temp table */
-      r3 = sqlite3ExprCodeTarget(pParse, pE2, r1);
-      sqlite3VdbeAddOp4(v, OP_MakeRecord, r3, 1, r2, &affinity, 1);
-      sqlite3VdbeAddOp4Int(v, OP_IdxInsert, iTab, r2, r3, 1);
+      sqlite3ExprCode(pParse, pE2, r1);
+      sqlite3VdbeAddOp4(v, OP_MakeRecord, r1, 1, r2, &affinity, 1);
+      sqlite3VdbeAddOp4Int(v, OP_IdxInsert, iTab, r2, r1, 1);
     }
     sqlite3ReleaseTempReg(pParse, r1);
     sqlite3ReleaseTempReg(pParse, r2);
@@ -99718,6 +101640,7 @@ SQLITE_PRIVATE void sqlite3CodeRhsOfIN(
     /* Subroutine return */
     sqlite3VdbeAddOp1(v, OP_Return, pExpr->y.sub.regReturn);
     sqlite3VdbeChangeP1(v, pExpr->y.sub.iAddr-1, sqlite3VdbeCurrentAddr(v)-1);
+    sqlite3ClearTempRegCache(pParse);
   }
 }
 #endif /* SQLITE_OMIT_SUBQUERY */
@@ -99731,7 +101654,7 @@ SQLITE_PRIVATE void sqlite3CodeRhsOfIN(
 **
 ** The pExpr parameter is the SELECT or EXISTS operator to be coded.
 **
-** The register that holds the result.  For a multi-column SELECT, 
+** Return the register that holds the result.  For a multi-column SELECT, 
 ** the result is stored in a contiguous array of registers and the
 ** return value is the register of the left-most result column.
 ** Return 0 if an error occurs.
@@ -99809,11 +101732,21 @@ SQLITE_PRIVATE int sqlite3CodeSubselect(Parse *pParse, Expr *pExpr){
     sqlite3VdbeAddOp2(v, OP_Integer, 0, dest.iSDParm);
     VdbeComment((v, "Init EXISTS result"));
   }
-  pLimit = sqlite3ExprAlloc(pParse->db, TK_INTEGER,&sqlite3IntTokens[1], 0);
   if( pSel->pLimit ){
-    sqlite3ExprDelete(pParse->db, pSel->pLimit->pLeft);
+    /* The subquery already has a limit.  If the pre-existing limit is X
+    ** then make the new limit X<>0 so that the new limit is either 1 or 0 */
+    sqlite3 *db = pParse->db;
+    pLimit = sqlite3Expr(db, TK_INTEGER, "0");
+    if( pLimit ){
+      pLimit->affExpr = SQLITE_AFF_NUMERIC;
+      pLimit = sqlite3PExpr(pParse, TK_NE,
+                            sqlite3ExprDup(db, pSel->pLimit->pLeft, 0), pLimit);
+    }
+    sqlite3ExprDelete(db, pSel->pLimit->pLeft);
     pSel->pLimit->pLeft = pLimit;
   }else{
+    /* If there is no pre-existing limit add a limit of 1 */
+    pLimit = sqlite3Expr(pParse->db, TK_INTEGER, "1");
     pSel->pLimit = sqlite3PExpr(pParse, TK_LIMIT, pLimit, 0);
   }
   pSel->iLimit = 0;
@@ -99828,6 +101761,7 @@ SQLITE_PRIVATE int sqlite3CodeSubselect(Parse *pParse, Expr *pExpr){
     /* Subroutine return */
     sqlite3VdbeAddOp1(v, OP_Return, pExpr->y.sub.regReturn);
     sqlite3VdbeChangeP1(v, pExpr->y.sub.iAddr-1, sqlite3VdbeCurrentAddr(v)-1);
+    sqlite3ClearTempRegCache(pParse);
   }
 
   return rReg;
@@ -99975,29 +101909,43 @@ static void sqlite3ExprCodeIN(
     int r2, regToFree;
     int regCkNull = 0;
     int ii;
+    int bLhsReal;  /* True if the LHS of the IN has REAL affinity */
     assert( !ExprHasProperty(pExpr, EP_xIsSelect) );
     if( destIfNull!=destIfFalse ){
       regCkNull = sqlite3GetTempReg(pParse);
       sqlite3VdbeAddOp3(v, OP_BitAnd, rLhs, rLhs, regCkNull);
     }
+    bLhsReal = sqlite3ExprAffinity(pExpr->pLeft)==SQLITE_AFF_REAL;
     for(ii=0; ii<pList->nExpr; ii++){
-      r2 = sqlite3ExprCodeTemp(pParse, pList->a[ii].pExpr, &regToFree);
+      if( bLhsReal ){
+        r2 = regToFree = sqlite3GetTempReg(pParse);
+        sqlite3ExprCode(pParse, pList->a[ii].pExpr, r2);
+        sqlite3VdbeAddOp4(v, OP_Affinity, r2, 1, 0, "E", P4_STATIC);
+      }else{
+        r2 = sqlite3ExprCodeTemp(pParse, pList->a[ii].pExpr, &regToFree);
+      }
       if( regCkNull && sqlite3ExprCanBeNull(pList->a[ii].pExpr) ){
         sqlite3VdbeAddOp3(v, OP_BitAnd, regCkNull, r2, regCkNull);
       }
+      sqlite3ReleaseTempReg(pParse, regToFree);
       if( ii<pList->nExpr-1 || destIfNull!=destIfFalse ){
-        sqlite3VdbeAddOp4(v, OP_Eq, rLhs, labelOk, r2,
+        int op = rLhs!=r2 ? OP_Eq : OP_NotNull;
+        sqlite3VdbeAddOp4(v, op, rLhs, labelOk, r2,
                           (void*)pColl, P4_COLLSEQ);
-        VdbeCoverageIf(v, ii<pList->nExpr-1);
-        VdbeCoverageIf(v, ii==pList->nExpr-1);
+        VdbeCoverageIf(v, ii<pList->nExpr-1 && op==OP_Eq);
+        VdbeCoverageIf(v, ii==pList->nExpr-1 && op==OP_Eq);
+        VdbeCoverageIf(v, ii<pList->nExpr-1 && op==OP_NotNull);
+        VdbeCoverageIf(v, ii==pList->nExpr-1 && op==OP_NotNull);
         sqlite3VdbeChangeP5(v, zAff[0]);
       }else{
+        int op = rLhs!=r2 ? OP_Ne : OP_IsNull;
         assert( destIfNull==destIfFalse );
-        sqlite3VdbeAddOp4(v, OP_Ne, rLhs, destIfFalse, r2,
-                          (void*)pColl, P4_COLLSEQ); VdbeCoverage(v);
+        sqlite3VdbeAddOp4(v, op, rLhs, destIfFalse, r2,
+                          (void*)pColl, P4_COLLSEQ);
+        VdbeCoverageIf(v, op==OP_Ne);
+        VdbeCoverageIf(v, op==OP_IsNull);
         sqlite3VdbeChangeP5(v, zAff[0] | SQLITE_JUMPIFNULL);
       }
-      sqlite3ReleaseTempReg(pParse, regToFree);
     }
     if( regCkNull ){
       sqlite3VdbeAddOp2(v, OP_IsNull, regCkNull, destIfNull); VdbeCoverage(v);
@@ -100017,6 +101965,7 @@ static void sqlite3ExprCodeIN(
   }else{
     destStep2 = destStep6 = sqlite3VdbeMakeLabel(pParse);
   }
+  if( pParse->nErr ) goto sqlite3ExprCodeIN_finished;
   for(i=0; i<nVector; i++){
     Expr *p = sqlite3VectorFieldSubexpr(pExpr->pLeft, i);
     if( sqlite3ExprCanBeNull(p) ){
@@ -100198,16 +102147,45 @@ SQLITE_PRIVATE void sqlite3ExprCodeLoadIndexColumn(
   }
 }
 
+#ifndef SQLITE_OMIT_GENERATED_COLUMNS
+/*
+** Generate code that will compute the value of generated column pCol
+** and store the result in register regOut
+*/
+SQLITE_PRIVATE void sqlite3ExprCodeGeneratedColumn(
+  Parse *pParse,
+  Column *pCol,
+  int regOut
+){
+  int iAddr;
+  Vdbe *v = pParse->pVdbe;
+  assert( v!=0 );
+  assert( pParse->iSelfTab!=0 );
+  if( pParse->iSelfTab>0 ){
+    iAddr = sqlite3VdbeAddOp3(v, OP_IfNullRow, pParse->iSelfTab-1, 0, regOut);
+  }else{
+    iAddr = 0;
+  }
+  sqlite3ExprCode(pParse, pCol->pDflt, regOut);
+  if( pCol->affinity>=SQLITE_AFF_TEXT ){
+    sqlite3VdbeAddOp4(v, OP_Affinity, regOut, 1, 0, &pCol->affinity, 1);
+  }
+  if( iAddr ) sqlite3VdbeJumpHere(v, iAddr);
+}
+#endif /* SQLITE_OMIT_GENERATED_COLUMNS */
+
 /*
 ** Generate code to extract the value of the iCol-th column of a table.
 */
 SQLITE_PRIVATE void sqlite3ExprCodeGetColumnOfTable(
-  Vdbe *v,        /* The VDBE under construction */
+  Vdbe *v,        /* Parsing context */
   Table *pTab,    /* The table containing the value */
   int iTabCur,    /* The table cursor.  Or the PK cursor for WITHOUT ROWID */
   int iCol,       /* Index of the column to extract */
   int regOut      /* Extract the value into this register */
 ){
+  Column *pCol;
+  assert( v!=0 );
   if( pTab==0 ){
     sqlite3VdbeAddOp3(v, OP_Column, iTabCur, iCol, regOut);
     return;
@@ -100215,14 +102193,36 @@ SQLITE_PRIVATE void sqlite3ExprCodeGetColumnOfTable(
   if( iCol<0 || iCol==pTab->iPKey ){
     sqlite3VdbeAddOp2(v, OP_Rowid, iTabCur, regOut);
   }else{
-    int op = IsVirtual(pTab) ? OP_VColumn : OP_Column;
-    int x = iCol;
-    if( !HasRowid(pTab) && !IsVirtual(pTab) ){
-      x = sqlite3ColumnOfIndex(sqlite3PrimaryKeyIndex(pTab), iCol);
+    int op;
+    int x;
+    if( IsVirtual(pTab) ){
+      op = OP_VColumn;
+      x = iCol;
+#ifndef SQLITE_OMIT_GENERATED_COLUMNS
+    }else if( (pCol = &pTab->aCol[iCol])->colFlags & COLFLAG_VIRTUAL ){
+      Parse *pParse = sqlite3VdbeParser(v);
+      if( pCol->colFlags & COLFLAG_BUSY ){
+        sqlite3ErrorMsg(pParse, "generated column loop on \"%s\"", pCol->zName);
+      }else{
+        int savedSelfTab = pParse->iSelfTab;
+        pCol->colFlags |= COLFLAG_BUSY;
+        pParse->iSelfTab = iTabCur+1;
+        sqlite3ExprCodeGeneratedColumn(pParse, pCol, regOut);
+        pParse->iSelfTab = savedSelfTab;
+        pCol->colFlags &= ~COLFLAG_BUSY;
+      }
+      return;
+#endif
+    }else if( !HasRowid(pTab) ){
+      testcase( iCol!=sqlite3TableColumnToStorage(pTab, iCol) );
+      x = sqlite3TableColumnToIndex(sqlite3PrimaryKeyIndex(pTab), iCol);
+      op = OP_Column;
+    }else{
+      x = sqlite3TableColumnToStorage(pTab,iCol);
+      testcase( x!=iCol );
+      op = OP_Column;
     }
     sqlite3VdbeAddOp3(v, op, iTabCur, x, regOut);
-  }
-  if( iCol>=0 ){
     sqlite3ColumnDefault(v, pTab, iCol, regOut);
   }
 }
@@ -100242,11 +102242,11 @@ SQLITE_PRIVATE int sqlite3ExprCodeGetColumn(
   int iReg,        /* Store results here */
   u8 p5            /* P5 value for OP_Column + FLAGS */
 ){
-  Vdbe *v = pParse->pVdbe;
-  assert( v!=0 );
-  sqlite3ExprCodeGetColumnOfTable(v, pTab, iTable, iColumn, iReg);
+  assert( pParse->pVdbe!=0 );
+  sqlite3ExprCodeGetColumnOfTable(pParse->pVdbe, pTab, iTable, iColumn, iReg);
   if( p5 ){
-    sqlite3VdbeChangeP5(v, p5);
+    VdbeOp *pOp = sqlite3VdbeGetOp(pParse->pVdbe,-1);
+    if( pOp->opcode==OP_Column ) pOp->p5 = p5;
   }
   return iReg;
 }
@@ -100256,7 +102256,6 @@ SQLITE_PRIVATE int sqlite3ExprCodeGetColumn(
 ** over to iTo..iTo+nReg-1.
 */
 SQLITE_PRIVATE void sqlite3ExprCodeMove(Parse *pParse, int iFrom, int iTo, int nReg){
-  assert( iFrom>=iTo+nReg || iFrom+nReg<=iTo );
   sqlite3VdbeAddOp3(pParse->pVdbe, OP_Move, iFrom, iTo, nReg);
 }
 
@@ -100265,7 +102264,8 @@ SQLITE_PRIVATE void sqlite3ExprCodeMove(Parse *pParse, int iFrom, int iTo, int n
 ** register iReg.  The caller must ensure that iReg already contains
 ** the correct value for the expression.
 */
-static void exprToRegister(Expr *p, int iReg){
+static void exprToRegister(Expr *pExpr, int iReg){
+  Expr *p = sqlite3ExprSkipCollateAndLikely(pExpr);
   p->op2 = p->op;
   p->op = TK_REGISTER;
   p->iTable = iReg;
@@ -100307,6 +102307,109 @@ static int exprCodeVector(Parse *pParse, Expr *p, int *piFreeable){
   return iResult;
 }
 
+/*
+** Generate code to implement special SQL functions that are implemented
+** in-line rather than by using the usual callbacks.
+*/
+static int exprCodeInlineFunction(
+  Parse *pParse,        /* Parsing context */
+  ExprList *pFarg,      /* List of function arguments */
+  int iFuncId,          /* Function ID.  One of the INTFUNC_... values */
+  int target            /* Store function result in this register */
+){
+  int nFarg;
+  Vdbe *v = pParse->pVdbe;
+  assert( v!=0 );
+  assert( pFarg!=0 );
+  nFarg = pFarg->nExpr;
+  assert( nFarg>0 );  /* All in-line functions have at least one argument */
+  switch( iFuncId ){
+    case INLINEFUNC_coalesce: {
+      /* Attempt a direct implementation of the built-in COALESCE() and
+      ** IFNULL() functions.  This avoids unnecessary evaluation of
+      ** arguments past the first non-NULL argument.
+      */
+      int endCoalesce = sqlite3VdbeMakeLabel(pParse);
+      int i;
+      assert( nFarg>=2 );
+      sqlite3ExprCode(pParse, pFarg->a[0].pExpr, target);
+      for(i=1; i<nFarg; i++){
+        sqlite3VdbeAddOp2(v, OP_NotNull, target, endCoalesce);
+        VdbeCoverage(v);
+        sqlite3ExprCode(pParse, pFarg->a[i].pExpr, target);
+      }
+      if( sqlite3VdbeGetOp(v, -1)->opcode==OP_Copy ){
+        sqlite3VdbeChangeP5(v, 1);  /* Tag trailing OP_Copy as not mergable */
+      }
+      sqlite3VdbeResolveLabel(v, endCoalesce);
+      break;
+    }
+
+    default: {   
+      /* The UNLIKELY() function is a no-op.  The result is the value
+      ** of the first argument.
+      */
+      assert( nFarg==1 || nFarg==2 );
+      target = sqlite3ExprCodeTarget(pParse, pFarg->a[0].pExpr, target);
+      break;
+    }
+
+  /***********************************************************************
+  ** Test-only SQL functions that are only usable if enabled
+  ** via SQLITE_TESTCTRL_INTERNAL_FUNCTIONS
+  */
+    case INLINEFUNC_expr_compare: {
+      /* Compare two expressions using sqlite3ExprCompare() */
+      assert( nFarg==2 );
+      sqlite3VdbeAddOp2(v, OP_Integer, 
+         sqlite3ExprCompare(0,pFarg->a[0].pExpr, pFarg->a[1].pExpr,-1),
+         target);
+      break;
+    }
+
+    case INLINEFUNC_expr_implies_expr: {
+      /* Compare two expressions using sqlite3ExprImpliesExpr() */
+      assert( nFarg==2 );
+      sqlite3VdbeAddOp2(v, OP_Integer, 
+         sqlite3ExprImpliesExpr(pParse,pFarg->a[0].pExpr, pFarg->a[1].pExpr,-1),
+         target);
+      break;
+    }
+
+    case INLINEFUNC_implies_nonnull_row: {
+      /* REsult of sqlite3ExprImpliesNonNullRow() */
+      Expr *pA1;
+      assert( nFarg==2 );
+      pA1 = pFarg->a[1].pExpr;
+      if( pA1->op==TK_COLUMN ){
+        sqlite3VdbeAddOp2(v, OP_Integer, 
+           sqlite3ExprImpliesNonNullRow(pFarg->a[0].pExpr,pA1->iTable),
+           target);
+      }else{
+        sqlite3VdbeAddOp2(v, OP_Null, 0, target);
+      }
+      break;
+    }
+
+#ifdef SQLITE_DEBUG
+    case INLINEFUNC_affinity: {
+      /* The AFFINITY() function evaluates to a string that describes
+      ** the type affinity of the argument.  This is used for testing of
+      ** the SQLite type logic.
+      */
+      const char *azAff[] = { "blob", "text", "numeric", "integer", "real" };
+      char aff;
+      assert( nFarg==1 );
+      aff = sqlite3ExprAffinity(pFarg->a[0].pExpr);
+      sqlite3VdbeLoadString(v, target, 
+              (aff<=SQLITE_AFF_NONE) ? "none" : azAff[aff-SQLITE_AFF_BLOB]);
+      break;
+    }
+#endif
+  }
+  return target;
+}
+
 
 /*
 ** Generate code into the current Vdbe to evaluate the given
@@ -100357,6 +102460,7 @@ expr_code_doover:
     }
     case TK_COLUMN: {
       int iTab = pExpr->iTable;
+      int iReg;
       if( ExprHasProperty(pExpr, EP_FixedCol) ){
         /* This COLUMN expression is really a constant due to WHERE clause
         ** constraints, and that constant is coded by the pExpr->pLeft
@@ -100364,9 +102468,14 @@ expr_code_doover:
         ** datatype by applying the Affinity of the table column to the
         ** constant.
         */
-        int iReg = sqlite3ExprCodeTarget(pParse, pExpr->pLeft,target);
-        int aff = sqlite3TableColumnAffinity(pExpr->y.pTab, pExpr->iColumn);
-        if( aff!=SQLITE_AFF_BLOB ){
+        int aff;
+        iReg = sqlite3ExprCodeTarget(pParse, pExpr->pLeft,target);
+        if( pExpr->y.pTab ){
+          aff = sqlite3TableColumnAffinity(pExpr->y.pTab, pExpr->iColumn);
+        }else{
+          aff = pExpr->affExpr;
+        }
+        if( aff>SQLITE_AFF_BLOB ){
           static const char zAff[] = "B\000C\000D\000E";
           assert( SQLITE_AFF_BLOB=='A' );
           assert( SQLITE_AFF_TEXT=='B' );
@@ -100381,17 +102490,60 @@ expr_code_doover:
       }
       if( iTab<0 ){
         if( pParse->iSelfTab<0 ){
-          /* Generating CHECK constraints or inserting into partial index */
-          return pExpr->iColumn - pParse->iSelfTab;
+          /* Other columns in the same row for CHECK constraints or
+          ** generated columns or for inserting into partial index.
+          ** The row is unpacked into registers beginning at
+          ** 0-(pParse->iSelfTab).  The rowid (if any) is in a register
+          ** immediately prior to the first column.
+          */
+          Column *pCol;
+          Table *pTab = pExpr->y.pTab;
+          int iSrc;
+          int iCol = pExpr->iColumn;
+          assert( pTab!=0 );
+          assert( iCol>=XN_ROWID );
+          assert( iCol<pTab->nCol );
+          if( iCol<0 ){
+            return -1-pParse->iSelfTab;
+          }
+          pCol = pTab->aCol + iCol;
+          testcase( iCol!=sqlite3TableColumnToStorage(pTab,iCol) );
+          iSrc = sqlite3TableColumnToStorage(pTab, iCol) - pParse->iSelfTab;
+#ifndef SQLITE_OMIT_GENERATED_COLUMNS
+          if( pCol->colFlags & COLFLAG_GENERATED ){
+            if( pCol->colFlags & COLFLAG_BUSY ){
+              sqlite3ErrorMsg(pParse, "generated column loop on \"%s\"",
+                              pCol->zName);
+              return 0;
+            }
+            pCol->colFlags |= COLFLAG_BUSY;
+            if( pCol->colFlags & COLFLAG_NOTAVAIL ){
+              sqlite3ExprCodeGeneratedColumn(pParse, pCol, iSrc);
+            }
+            pCol->colFlags &= ~(COLFLAG_BUSY|COLFLAG_NOTAVAIL);
+            return iSrc;
+          }else
+#endif /* SQLITE_OMIT_GENERATED_COLUMNS */
+          if( pCol->affinity==SQLITE_AFF_REAL ){
+            sqlite3VdbeAddOp2(v, OP_SCopy, iSrc, target);
+            sqlite3VdbeAddOp1(v, OP_RealAffinity, target);
+            return target;
+          }else{
+            return iSrc;
+          }
         }else{
           /* Coding an expression that is part of an index where column names
           ** in the index refer to the table to which the index belongs */
           iTab = pParse->iSelfTab - 1;
         }
       }
-      return sqlite3ExprCodeGetColumn(pParse, pExpr->y.pTab,
+      iReg = sqlite3ExprCodeGetColumn(pParse, pExpr->y.pTab,
                                pExpr->iColumn, iTab, target,
                                pExpr->op2);
+      if( pExpr->y.pTab==0 && pExpr->affExpr==SQLITE_AFF_REAL ){
+        sqlite3VdbeAddOp1(v, OP_RealAffinity, iReg);
+      }
+      return iReg;
     }
     case TK_INTEGER: {
       codeInteger(pParse, pExpr, 0, target);
@@ -100413,7 +102565,12 @@ expr_code_doover:
       sqlite3VdbeLoadString(v, target, pExpr->u.zToken);
       return target;
     }
-    case TK_NULL: {
+    default: {
+      /* Make NULL the default case so that if a bug causes an illegal
+      ** Expr node to be passed into this function, it will be handled
+      ** sanely and not crash.  But keep the assert() to bring the problem
+      ** to the attention of the developers. */
+      assert( op==TK_NULL );
       sqlite3VdbeAddOp2(v, OP_Null, 0, target);
       return target;
     }
@@ -100440,7 +102597,7 @@ expr_code_doover:
       sqlite3VdbeAddOp2(v, OP_Variable, pExpr->iColumn, target);
       if( pExpr->u.zToken[1]!=0 ){
         const char *z = sqlite3VListNumToName(pParse->pVList, pExpr->iColumn);
-        assert( pExpr->u.zToken[0]=='?' || strcmp(pExpr->u.zToken, z)==0 );
+        assert( pExpr->u.zToken[0]=='?' || (z && !strcmp(pExpr->u.zToken, z)) );
         pParse->pVList[0] = 0; /* Indicate VList may no longer be enlarged */
         sqlite3VdbeAppendP4(v, (char*)z, P4_STATIC);
       }
@@ -100480,7 +102637,8 @@ expr_code_doover:
         r1 = sqlite3ExprCodeTemp(pParse, pLeft, &regFree1);
         r2 = sqlite3ExprCodeTemp(pParse, pExpr->pRight, &regFree2);
         codeCompare(pParse, pLeft, pExpr->pRight, op,
-            r1, r2, inReg, SQLITE_STOREP2 | p5);
+            r1, r2, inReg, SQLITE_STOREP2 | p5,
+            ExprHasProperty(pExpr,EP_Commuted));
         assert(TK_LT==OP_Lt); testcase(op==OP_Lt); VdbeCoverageIf(v,op==OP_Lt);
         assert(TK_LE==OP_Le); testcase(op==OP_Le); VdbeCoverageIf(v,op==OP_Le);
         assert(TK_GT==OP_Gt); testcase(op==OP_Gt); VdbeCoverageIf(v,op==OP_Gt);
@@ -100632,47 +102790,14 @@ expr_code_doover:
         sqlite3ErrorMsg(pParse, "unknown function: %s()", zId);
         break;
       }
-
-      /* Attempt a direct implementation of the built-in COALESCE() and
-      ** IFNULL() functions.  This avoids unnecessary evaluation of
-      ** arguments past the first non-NULL argument.
-      */
-      if( pDef->funcFlags & SQLITE_FUNC_COALESCE ){
-        int endCoalesce = sqlite3VdbeMakeLabel(pParse);
-        assert( nFarg>=2 );
-        sqlite3ExprCode(pParse, pFarg->a[0].pExpr, target);
-        for(i=1; i<nFarg; i++){
-          sqlite3VdbeAddOp2(v, OP_NotNull, target, endCoalesce);
-          VdbeCoverage(v);
-          sqlite3ExprCode(pParse, pFarg->a[i].pExpr, target);
-        }
-        sqlite3VdbeResolveLabel(v, endCoalesce);
-        break;
-      }
-
-      /* The UNLIKELY() function is a no-op.  The result is the value
-      ** of the first argument.
-      */
-      if( pDef->funcFlags & SQLITE_FUNC_UNLIKELY ){
-        assert( nFarg>=1 );
-        return sqlite3ExprCodeTarget(pParse, pFarg->a[0].pExpr, target);
-      }
-
-#ifdef SQLITE_DEBUG
-      /* The AFFINITY() function evaluates to a string that describes
-      ** the type affinity of the argument.  This is used for testing of
-      ** the SQLite type logic.
-      */
-      if( pDef->funcFlags & SQLITE_FUNC_AFFINITY ){
-        const char *azAff[] = { "blob", "text", "numeric", "integer", "real" };
-        char aff;
-        assert( nFarg==1 );
-        aff = sqlite3ExprAffinity(pFarg->a[0].pExpr);
-        sqlite3VdbeLoadString(v, target, 
-                              aff ? azAff[aff-SQLITE_AFF_BLOB] : "none");
-        return target;
+      if( pDef->funcFlags & SQLITE_FUNC_INLINE ){
+        assert( (pDef->funcFlags & SQLITE_FUNC_UNSAFE)==0 );
+        assert( (pDef->funcFlags & SQLITE_FUNC_DIRECT)==0 );
+        return exprCodeInlineFunction(pParse, pFarg,
+             SQLITE_PTR_TO_INT(pDef->pUserData), target);
+      }else if( pDef->funcFlags & (SQLITE_FUNC_DIRECT|SQLITE_FUNC_UNSAFE) ){
+        sqlite3ExprFunctionUsable(pParse, pExpr, pDef);
       }
-#endif
 
       for(i=0; i<nFarg; i++){
         if( i<32 && sqlite3ExprIsConstant(pFarg->a[i].pExpr) ){
@@ -100749,12 +102874,15 @@ expr_code_doover:
       }else
 #endif
       {
-        sqlite3VdbeAddOp4(v, pParse->iSelfTab ? OP_PureFunc0 : OP_Function0,
-                          constMask, r1, target, (char*)pDef, P4_FUNCDEF);
-        sqlite3VdbeChangeP5(v, (u8)nFarg);
+        sqlite3VdbeAddFunctionCall(pParse, constMask, r1, target, nFarg,
+                                   pDef, pExpr->op2);
       }
-      if( nFarg && constMask==0 ){
-        sqlite3ReleaseTempRange(pParse, r1, nFarg);
+      if( nFarg ){
+        if( constMask==0 ){
+          sqlite3ReleaseTempRange(pParse, r1, nFarg);
+        }else{
+          sqlite3VdbeReleaseRegisters(pParse, r1, nFarg, constMask, 1);
+        }
       }
       return target;
     }
@@ -100777,8 +102905,8 @@ expr_code_doover:
         pExpr->pLeft->iTable = sqlite3CodeSubselect(pParse, pExpr->pLeft);
       }
       assert( pExpr->iTable==0 || pExpr->pLeft->op==TK_SELECT );
-      if( pExpr->iTable
-       && pExpr->iTable!=(n = sqlite3ExprVectorSize(pExpr->pLeft)) 
+      if( pExpr->iTable!=0
+       && pExpr->iTable!=(n = sqlite3ExprVectorSize(pExpr->pLeft))
       ){
         sqlite3ErrorMsg(pParse, "%d columns assigned %d values",
                                 pExpr->iTable, n);
@@ -100848,17 +102976,19 @@ expr_code_doover:
       **   p1==2   ->    old.b         p1==5   ->    new.b       
       */
       Table *pTab = pExpr->y.pTab;
-      int p1 = pExpr->iTable * (pTab->nCol+1) + 1 + pExpr->iColumn;
+      int iCol = pExpr->iColumn;
+      int p1 = pExpr->iTable * (pTab->nCol+1) + 1 
+                     + sqlite3TableColumnToStorage(pTab, iCol);
 
       assert( pExpr->iTable==0 || pExpr->iTable==1 );
-      assert( pExpr->iColumn>=-1 && pExpr->iColumn<pTab->nCol );
-      assert( pTab->iPKey<0 || pExpr->iColumn!=pTab->iPKey );
+      assert( iCol>=-1 && iCol<pTab->nCol );
+      assert( pTab->iPKey<0 || iCol!=pTab->iPKey );
       assert( p1>=0 && p1<(pTab->nCol*2+2) );
 
       sqlite3VdbeAddOp2(v, OP_Param, p1, target);
       VdbeComment((v, "r[%d]=%s.%s", target,
         (pExpr->iTable ? "new" : "old"),
-        (pExpr->iColumn<0 ? "rowid" : pExpr->y.pTab->aCol[pExpr->iColumn].zName)
+        (pExpr->iColumn<0 ? "rowid" : pExpr->y.pTab->aCol[iCol].zName)
       ));
 
 #ifndef SQLITE_OMIT_FLOATING_POINT
@@ -100867,9 +102997,7 @@ expr_code_doover:
       **
       ** EVIDENCE-OF: R-60985-57662 SQLite will convert the value back to
       ** floating point when extracting it from the record.  */
-      if( pExpr->iColumn>=0 
-       && pTab->aCol[pExpr->iColumn].affinity==SQLITE_AFF_REAL
-      ){
+      if( iCol>=0 && pTab->aCol[iCol].affinity==SQLITE_AFF_REAL ){
         sqlite3VdbeAddOp1(v, OP_RealAffinity, target);
       }
 #endif
@@ -100881,10 +103009,23 @@ expr_code_doover:
       break;
     }
 
+    /* TK_IF_NULL_ROW Expr nodes are inserted ahead of expressions
+    ** that derive from the right-hand table of a LEFT JOIN.  The
+    ** Expr.iTable value is the table number for the right-hand table.
+    ** The expression is only evaluated if that table is not currently
+    ** on a LEFT JOIN NULL row.
+    */
     case TK_IF_NULL_ROW: {
       int addrINR;
+      u8 okConstFactor = pParse->okConstFactor;
       addrINR = sqlite3VdbeAddOp1(v, OP_IfNullRow, pExpr->iTable);
+      /* Temporarily disable factoring of constant expressions, since
+      ** even though expressions may appear to be constant, they are not
+      ** really constant because they originate from the right-hand side
+      ** of a LEFT JOIN. */
+      pParse->okConstFactor = 0;
       inReg = sqlite3ExprCodeTarget(pParse, pExpr->pLeft, target);
+      pParse->okConstFactor = okConstFactor;
       sqlite3VdbeJumpHere(v, addrINR);
       sqlite3VdbeChangeP3(v, addrINR, inReg);
       break;
@@ -100911,7 +103052,7 @@ expr_code_doover:
     ** or if there is no matching Ei, the ELSE term Y, or if there is
     ** no ELSE term, NULL.
     */
-    default: assert( op==TK_CASE ); {
+    case TK_CASE: {
       int endLabel;                     /* GOTO label for end of CASE stmt */
       int nextCase;                     /* GOTO label for next WHEN clause */
       int nExpr;                        /* 2x number of WHEN terms */
@@ -100921,6 +103062,8 @@ expr_code_doover:
       Expr opCompare;                   /* The X==Ei expression */
       Expr *pX;                         /* The X expression */
       Expr *pTest = 0;                  /* X==Ei (form A) or just Ei (form B) */
+      Expr *pDel = 0;
+      sqlite3 *db = pParse->db;
 
       assert( !ExprHasProperty(pExpr, EP_xIsSelect) && pExpr->x.pList );
       assert(pExpr->x.pList->nExpr > 0);
@@ -100929,13 +103072,17 @@ expr_code_doover:
       nExpr = pEList->nExpr;
       endLabel = sqlite3VdbeMakeLabel(pParse);
       if( (pX = pExpr->pLeft)!=0 ){
-        exprNodeCopy(&tempX, pX);
+        pDel = sqlite3ExprDup(db, pX, 0);
+        if( db->mallocFailed ){
+          sqlite3ExprDelete(db, pDel);
+          break;
+        }
         testcase( pX->op==TK_COLUMN );
-        exprToRegister(&tempX, exprCodeVector(pParse, &tempX, &regFree1));
+        exprToRegister(pDel, exprCodeVector(pParse, pDel, &regFree1));
         testcase( regFree1==0 );
         memset(&opCompare, 0, sizeof(opCompare));
         opCompare.op = TK_EQ;
-        opCompare.pLeft = &tempX;
+        opCompare.pLeft = pDel;
         pTest = &opCompare;
         /* Ticket b351d95f9cd5ef17e9d9dbae18f5ca8611190001:
         ** The value in regFree1 might get SCopy-ed into the file result.
@@ -100963,32 +103110,33 @@ expr_code_doover:
       }else{
         sqlite3VdbeAddOp2(v, OP_Null, 0, target);
       }
+      sqlite3ExprDelete(db, pDel);
       sqlite3VdbeResolveLabel(v, endLabel);
       break;
     }
 #ifndef SQLITE_OMIT_TRIGGER
     case TK_RAISE: {
-      assert( pExpr->affinity==OE_Rollback 
-           || pExpr->affinity==OE_Abort
-           || pExpr->affinity==OE_Fail
-           || pExpr->affinity==OE_Ignore
+      assert( pExpr->affExpr==OE_Rollback 
+           || pExpr->affExpr==OE_Abort
+           || pExpr->affExpr==OE_Fail
+           || pExpr->affExpr==OE_Ignore
       );
       if( !pParse->pTriggerTab ){
         sqlite3ErrorMsg(pParse,
                        "RAISE() may only be used within a trigger-program");
         return 0;
       }
-      if( pExpr->affinity==OE_Abort ){
+      if( pExpr->affExpr==OE_Abort ){
         sqlite3MayAbort(pParse);
       }
       assert( !ExprHasProperty(pExpr, EP_IntValue) );
-      if( pExpr->affinity==OE_Ignore ){
+      if( pExpr->affExpr==OE_Ignore ){
         sqlite3VdbeAddOp4(
             v, OP_Halt, SQLITE_OK, OE_Ignore, 0, pExpr->u.zToken,0);
         VdbeCoverage(v);
       }else{
         sqlite3HaltConstraint(pParse, SQLITE_CONSTRAINT_TRIGGER,
-                              pExpr->affinity, pExpr->u.zToken, 0, 0);
+                              pExpr->affExpr, pExpr->u.zToken, 0, 0);
       }
 
       break;
@@ -101053,7 +103201,7 @@ SQLITE_PRIVATE int sqlite3ExprCodeAtInit(
 */
 SQLITE_PRIVATE int sqlite3ExprCodeTemp(Parse *pParse, Expr *pExpr, int *pReg){
   int r2;
-  pExpr = sqlite3ExprSkipCollate(pExpr);
+  pExpr = sqlite3ExprSkipCollateAndLikely(pExpr);
   if( ConstFactorOk(pParse)
    && pExpr->op!=TK_REGISTER
    && sqlite3ExprIsConstantNotJoin(pExpr)
@@ -101082,14 +103230,16 @@ SQLITE_PRIVATE void sqlite3ExprCode(Parse *pParse, Expr *pExpr, int target){
   int inReg;
 
   assert( target>0 && target<=pParse->nMem );
-  if( pExpr && pExpr->op==TK_REGISTER ){
-    sqlite3VdbeAddOp2(pParse->pVdbe, OP_Copy, pExpr->iTable, target);
-  }else{
-    inReg = sqlite3ExprCodeTarget(pParse, pExpr, target);
-    assert( pParse->pVdbe!=0 || pParse->db->mallocFailed );
-    if( inReg!=target && pParse->pVdbe ){
-      sqlite3VdbeAddOp2(pParse->pVdbe, OP_SCopy, inReg, target);
+  inReg = sqlite3ExprCodeTarget(pParse, pExpr, target);
+  assert( pParse->pVdbe!=0 || pParse->db->mallocFailed );
+  if( inReg!=target && pParse->pVdbe ){
+    u8 op;
+    if( ExprHasProperty(pExpr,EP_Subquery) ){
+      op = OP_Copy;
+    }else{
+      op = OP_SCopy;
     }
+    sqlite3VdbeAddOp2(pParse->pVdbe, op, inReg, target);
   }
 }
 
@@ -101119,30 +103269,6 @@ SQLITE_PRIVATE void sqlite3ExprCodeFactorable(Parse *pParse, Expr *pExpr, int ta
   }
 }
 
-/*
-** Generate code that evaluates the given expression and puts the result
-** in register target.
-**
-** Also make a copy of the expression results into another "cache" register
-** and modify the expression so that the next time it is evaluated,
-** the result is a copy of the cache register.
-**
-** This routine is used for expressions that are used multiple 
-** times.  They are evaluated once and the results of the expression
-** are reused.
-*/
-SQLITE_PRIVATE void sqlite3ExprCodeAndCache(Parse *pParse, Expr *pExpr, int target){
-  Vdbe *v = pParse->pVdbe;
-  int iMem;
-
-  assert( target>0 );
-  assert( pExpr->op!=TK_REGISTER );
-  sqlite3ExprCode(pParse, pExpr, target);
-  iMem = ++pParse->nMem;
-  sqlite3VdbeAddOp2(v, OP_Copy, target, iMem);
-  exprToRegister(pExpr, iMem);
-}
-
 /*
 ** Generate code that pushes the value of every element of the given
 ** expression list into a sequence of registers beginning at target.
@@ -101206,6 +103332,7 @@ SQLITE_PRIVATE int sqlite3ExprCodeExprList(
          && (pOp=sqlite3VdbeGetOp(v, -1))->opcode==OP_Copy
          && pOp->p1+pOp->p3+1==inReg
          && pOp->p2+pOp->p3+1==target+i
+         && pOp->p5==0  /* The do-not-merge flag must be clear */
         ){
           pOp->p3++;
         }else{
@@ -101244,40 +103371,44 @@ static void exprCodeBetween(
   void (*xJump)(Parse*,Expr*,int,int), /* Action to take */
   int jumpIfNull    /* Take the jump if the BETWEEN is NULL */
 ){
- Expr exprAnd;     /* The AND operator in  x>=y AND x<=z  */
+  Expr exprAnd;     /* The AND operator in  x>=y AND x<=z  */
   Expr compLeft;    /* The  x>=y  term */
   Expr compRight;   /* The  x<=z  term */
-  Expr exprX;       /* The  x  subexpression */
   int regFree1 = 0; /* Temporary use register */
+  Expr *pDel = 0;
+  sqlite3 *db = pParse->db;
 
   memset(&compLeft, 0, sizeof(Expr));
   memset(&compRight, 0, sizeof(Expr));
   memset(&exprAnd, 0, sizeof(Expr));
 
   assert( !ExprHasProperty(pExpr, EP_xIsSelect) );
-  exprNodeCopy(&exprX, pExpr->pLeft);
-  exprAnd.op = TK_AND;
-  exprAnd.pLeft = &compLeft;
-  exprAnd.pRight = &compRight;
-  compLeft.op = TK_GE;
-  compLeft.pLeft = &exprX;
-  compLeft.pRight = pExpr->x.pList->a[0].pExpr;
-  compRight.op = TK_LE;
-  compRight.pLeft = &exprX;
-  compRight.pRight = pExpr->x.pList->a[1].pExpr;
-  exprToRegister(&exprX, exprCodeVector(pParse, &exprX, &regFree1));
-  if( xJump ){
-    xJump(pParse, &exprAnd, dest, jumpIfNull);
-  }else{
-    /* Mark the expression is being from the ON or USING clause of a join
-    ** so that the sqlite3ExprCodeTarget() routine will not attempt to move
-    ** it into the Parse.pConstExpr list.  We should use a new bit for this,
-    ** for clarity, but we are out of bits in the Expr.flags field so we
-    ** have to reuse the EP_FromJoin bit.  Bummer. */
-    exprX.flags |= EP_FromJoin;
-    sqlite3ExprCodeTarget(pParse, &exprAnd, dest);
+  pDel = sqlite3ExprDup(db, pExpr->pLeft, 0);
+  if( db->mallocFailed==0 ){
+    exprAnd.op = TK_AND;
+    exprAnd.pLeft = &compLeft;
+    exprAnd.pRight = &compRight;
+    compLeft.op = TK_GE;
+    compLeft.pLeft = pDel;
+    compLeft.pRight = pExpr->x.pList->a[0].pExpr;
+    compRight.op = TK_LE;
+    compRight.pLeft = pDel;
+    compRight.pRight = pExpr->x.pList->a[1].pExpr;
+    exprToRegister(pDel, exprCodeVector(pParse, pDel, &regFree1));
+    if( xJump ){
+      xJump(pParse, &exprAnd, dest, jumpIfNull);
+    }else{
+      /* Mark the expression is being from the ON or USING clause of a join
+      ** so that the sqlite3ExprCodeTarget() routine will not attempt to move
+      ** it into the Parse.pConstExpr list.  We should use a new bit for this,
+      ** for clarity, but we are out of bits in the Expr.flags field so we
+      ** have to reuse the EP_FromJoin bit.  Bummer. */
+      pDel->flags |= EP_FromJoin;
+      sqlite3ExprCodeTarget(pParse, &exprAnd, dest);
+    }
+    sqlite3ReleaseTempReg(pParse, regFree1);
   }
-  sqlite3ReleaseTempReg(pParse, regFree1);
+  sqlite3ExprDelete(db, pDel);
 
   /* Ensure adequate test coverage */
   testcase( xJump==sqlite3ExprIfTrue  && jumpIfNull==0 && regFree1==0 );
@@ -101317,18 +103448,23 @@ SQLITE_PRIVATE void sqlite3ExprIfTrue(Parse *pParse, Expr *pExpr, int dest, int
   if( NEVER(pExpr==0) ) return;  /* No way this can happen */
   op = pExpr->op;
   switch( op ){
-    case TK_AND: {
-      int d2 = sqlite3VdbeMakeLabel(pParse);
-      testcase( jumpIfNull==0 );
-      sqlite3ExprIfFalse(pParse, pExpr->pLeft, d2,jumpIfNull^SQLITE_JUMPIFNULL);
-      sqlite3ExprIfTrue(pParse, pExpr->pRight, dest, jumpIfNull);
-      sqlite3VdbeResolveLabel(v, d2);
-      break;
-    }
+    case TK_AND:
     case TK_OR: {
-      testcase( jumpIfNull==0 );
-      sqlite3ExprIfTrue(pParse, pExpr->pLeft, dest, jumpIfNull);
-      sqlite3ExprIfTrue(pParse, pExpr->pRight, dest, jumpIfNull);
+      Expr *pAlt = sqlite3ExprSimplifiedAndOr(pExpr);
+      if( pAlt!=pExpr ){
+        sqlite3ExprIfTrue(pParse, pAlt, dest, jumpIfNull);
+      }else if( op==TK_AND ){
+        int d2 = sqlite3VdbeMakeLabel(pParse);
+        testcase( jumpIfNull==0 );
+        sqlite3ExprIfFalse(pParse, pExpr->pLeft, d2,
+                           jumpIfNull^SQLITE_JUMPIFNULL);
+        sqlite3ExprIfTrue(pParse, pExpr->pRight, dest, jumpIfNull);
+        sqlite3VdbeResolveLabel(v, d2);
+      }else{
+        testcase( jumpIfNull==0 );
+        sqlite3ExprIfTrue(pParse, pExpr->pLeft, dest, jumpIfNull);
+        sqlite3ExprIfTrue(pParse, pExpr->pRight, dest, jumpIfNull);
+      }
       break;
     }
     case TK_NOT: {
@@ -101371,7 +103507,7 @@ SQLITE_PRIVATE void sqlite3ExprIfTrue(Parse *pParse, Expr *pExpr, int dest, int
       r1 = sqlite3ExprCodeTemp(pParse, pExpr->pLeft, &regFree1);
       r2 = sqlite3ExprCodeTemp(pParse, pExpr->pRight, &regFree2);
       codeCompare(pParse, pExpr->pLeft, pExpr->pRight, op,
-                  r1, r2, dest, jumpIfNull);
+                  r1, r2, dest, jumpIfNull, ExprHasProperty(pExpr,EP_Commuted));
       assert(TK_LT==OP_Lt); testcase(op==OP_Lt); VdbeCoverageIf(v,op==OP_Lt);
       assert(TK_LE==OP_Le); testcase(op==OP_Le); VdbeCoverageIf(v,op==OP_Le);
       assert(TK_GT==OP_Gt); testcase(op==OP_Gt); VdbeCoverageIf(v,op==OP_Gt);
@@ -101414,9 +103550,9 @@ SQLITE_PRIVATE void sqlite3ExprIfTrue(Parse *pParse, Expr *pExpr, int dest, int
 #endif
     default: {
     default_expr:
-      if( exprAlwaysTrue(pExpr) ){
+      if( ExprAlwaysTrue(pExpr) ){
         sqlite3VdbeGoto(v, dest);
-      }else if( exprAlwaysFalse(pExpr) ){
+      }else if( ExprAlwaysFalse(pExpr) ){
         /* No-op */
       }else{
         r1 = sqlite3ExprCodeTemp(pParse, pExpr, &regFree1);
@@ -101484,18 +103620,23 @@ SQLITE_PRIVATE void sqlite3ExprIfFalse(Parse *pParse, Expr *pExpr, int dest, int
   assert( pExpr->op!=TK_GE || op==OP_Lt );
 
   switch( pExpr->op ){
-    case TK_AND: {
-      testcase( jumpIfNull==0 );
-      sqlite3ExprIfFalse(pParse, pExpr->pLeft, dest, jumpIfNull);
-      sqlite3ExprIfFalse(pParse, pExpr->pRight, dest, jumpIfNull);
-      break;
-    }
+    case TK_AND:
     case TK_OR: {
-      int d2 = sqlite3VdbeMakeLabel(pParse);
-      testcase( jumpIfNull==0 );
-      sqlite3ExprIfTrue(pParse, pExpr->pLeft, d2, jumpIfNull^SQLITE_JUMPIFNULL);
-      sqlite3ExprIfFalse(pParse, pExpr->pRight, dest, jumpIfNull);
-      sqlite3VdbeResolveLabel(v, d2);
+      Expr *pAlt = sqlite3ExprSimplifiedAndOr(pExpr);
+      if( pAlt!=pExpr ){
+        sqlite3ExprIfFalse(pParse, pAlt, dest, jumpIfNull);
+      }else if( pExpr->op==TK_AND ){
+        testcase( jumpIfNull==0 );
+        sqlite3ExprIfFalse(pParse, pExpr->pLeft, dest, jumpIfNull);
+        sqlite3ExprIfFalse(pParse, pExpr->pRight, dest, jumpIfNull);
+      }else{
+        int d2 = sqlite3VdbeMakeLabel(pParse);
+        testcase( jumpIfNull==0 );
+        sqlite3ExprIfTrue(pParse, pExpr->pLeft, d2,
+                          jumpIfNull^SQLITE_JUMPIFNULL);
+        sqlite3ExprIfFalse(pParse, pExpr->pRight, dest, jumpIfNull);
+        sqlite3VdbeResolveLabel(v, d2);
+      }
       break;
     }
     case TK_NOT: {
@@ -101541,7 +103682,7 @@ SQLITE_PRIVATE void sqlite3ExprIfFalse(Parse *pParse, Expr *pExpr, int dest, int
       r1 = sqlite3ExprCodeTemp(pParse, pExpr->pLeft, &regFree1);
       r2 = sqlite3ExprCodeTemp(pParse, pExpr->pRight, &regFree2);
       codeCompare(pParse, pExpr->pLeft, pExpr->pRight, op,
-                  r1, r2, dest, jumpIfNull);
+                  r1, r2, dest, jumpIfNull,ExprHasProperty(pExpr,EP_Commuted));
       assert(TK_LT==OP_Lt); testcase(op==OP_Lt); VdbeCoverageIf(v,op==OP_Lt);
       assert(TK_LE==OP_Le); testcase(op==OP_Le); VdbeCoverageIf(v,op==OP_Le);
       assert(TK_GT==OP_Gt); testcase(op==OP_Gt); VdbeCoverageIf(v,op==OP_Gt);
@@ -101584,9 +103725,9 @@ SQLITE_PRIVATE void sqlite3ExprIfFalse(Parse *pParse, Expr *pExpr, int dest, int
 #endif
     default: {
     default_expr: 
-      if( exprAlwaysFalse(pExpr) ){
+      if( ExprAlwaysFalse(pExpr) ){
         sqlite3VdbeGoto(v, dest);
-      }else if( exprAlwaysTrue(pExpr) ){
+      }else if( ExprAlwaysTrue(pExpr) ){
         /* no-op */
       }else{
         r1 = sqlite3ExprCodeTemp(pParse, pExpr, &regFree1);
@@ -101706,20 +103847,17 @@ SQLITE_PRIVATE int sqlite3ExprCompare(Parse *pParse, Expr *pA, Expr *pB, int iTa
     return 2;
   }
   if( pA->op!=TK_COLUMN && pA->op!=TK_AGG_COLUMN && pA->u.zToken ){
-    if( pA->op==TK_FUNCTION ){
+    if( pA->op==TK_FUNCTION || pA->op==TK_AGG_FUNCTION ){
       if( sqlite3StrICmp(pA->u.zToken,pB->u.zToken)!=0 ) return 2;
 #ifndef SQLITE_OMIT_WINDOWFUNC
-      /* Justification for the assert():
-      ** window functions have p->op==TK_FUNCTION but aggregate functions
-      ** have p->op==TK_AGG_FUNCTION.  So any comparison between an aggregate
-      ** function and a window function should have failed before reaching
-      ** this point.  And, it is not possible to have a window function and
-      ** a scalar function with the same name and number of arguments.  So
-      ** if we reach this point, either A and B both window functions or
-      ** neither are a window functions. */
-      assert( ExprHasProperty(pA,EP_WinFunc)==ExprHasProperty(pB,EP_WinFunc) );
+      assert( pA->op==pB->op );
+      if( ExprHasProperty(pA,EP_WinFunc)!=ExprHasProperty(pB,EP_WinFunc) ){
+        return 2;
+      }
       if( ExprHasProperty(pA,EP_WinFunc) ){
-        if( sqlite3WindowCompare(pParse,pA->y.pWin,pB->y.pWin)!=0 ) return 2;
+        if( sqlite3WindowCompare(pParse, pA->y.pWin, pB->y.pWin, 1)!=0 ){
+          return 2;
+        }
       }
 #endif
     }else if( pA->op==TK_NULL ){
@@ -101730,7 +103868,8 @@ SQLITE_PRIVATE int sqlite3ExprCompare(Parse *pParse, Expr *pA, Expr *pB, int iTa
       return 2;
     }
   }
-  if( (pA->flags & EP_Distinct)!=(pB->flags & EP_Distinct) ) return 2;
+  if( (pA->flags & (EP_Distinct|EP_Commuted))
+     != (pB->flags & (EP_Distinct|EP_Commuted)) ) return 2;
   if( (combinedFlags & EP_TokenOnly)==0 ){
     if( combinedFlags & EP_xIsSelect ) return 2;
     if( (combinedFlags & EP_FixedCol)==0
@@ -101742,16 +103881,33 @@ SQLITE_PRIVATE int sqlite3ExprCompare(Parse *pParse, Expr *pA, Expr *pB, int iTa
      && (combinedFlags & EP_Reduced)==0
     ){
       if( pA->iColumn!=pB->iColumn ) return 2;
-      if( pA->iTable!=pB->iTable 
-       && (pA->iTable!=iTab || NEVER(pB->iTable>=0)) ) return 2;
+      if( pA->op2!=pB->op2 ){
+        if( pA->op==TK_TRUTH ) return 2;
+        if( pA->op==TK_FUNCTION && iTab<0 ){
+          /* Ex: CREATE TABLE t1(a CHECK( a<julianday('now') ));
+          **     INSERT INTO t1(a) VALUES(julianday('now')+10);
+          ** Without this test, sqlite3ExprCodeAtInit() will run on the
+          ** the julianday() of INSERT first, and remember that expression.
+          ** Then sqlite3ExprCodeInit() will see the julianday() in the CHECK
+          ** constraint as redundant, reusing the one from the INSERT, even
+          ** though the julianday() in INSERT lacks the critical NC_IsCheck
+          ** flag.  See ticket [830277d9db6c3ba1] (2019-10-30)
+          */
+          return 2;
+        }
+      }
+      if( pA->op!=TK_IN && pA->iTable!=pB->iTable && pA->iTable!=iTab ){
+        return 2;
+      }
     }
   }
   return 0;
 }
 
 /*
-** Compare two ExprList objects.  Return 0 if they are identical and 
-** non-zero if they differ in any way.
+** Compare two ExprList objects.  Return 0 if they are identical, 1
+** if they are certainly different, or 2 if it is not possible to 
+** determine if they are identical or not.
 **
 ** If any subelement of pB has Expr.iTable==(-1) then it is allowed
 ** to compare equal to an equivalent element in pA with Expr.iTable==iTab.
@@ -101770,10 +103926,11 @@ SQLITE_PRIVATE int sqlite3ExprListCompare(ExprList *pA, ExprList *pB, int iTab){
   if( pA==0 || pB==0 ) return 1;
   if( pA->nExpr!=pB->nExpr ) return 1;
   for(i=0; i<pA->nExpr; i++){
+    int res;
     Expr *pExprA = pA->a[i].pExpr;
     Expr *pExprB = pB->a[i].pExpr;
-    if( pA->a[i].sortOrder!=pB->a[i].sortOrder ) return 1;
-    if( sqlite3ExprCompare(0, pExprA, pExprB, iTab) ) return 1;
+    if( pA->a[i].sortFlags!=pB->a[i].sortFlags ) return 1;
+    if( (res = sqlite3ExprCompare(0, pExprA, pExprB, iTab)) ) return res;
   }
   return 0;
 }
@@ -101784,11 +103941,88 @@ SQLITE_PRIVATE int sqlite3ExprListCompare(ExprList *pA, ExprList *pB, int iTab){
 */
 SQLITE_PRIVATE int sqlite3ExprCompareSkip(Expr *pA, Expr *pB, int iTab){
   return sqlite3ExprCompare(0,
-             sqlite3ExprSkipCollate(pA),
-             sqlite3ExprSkipCollate(pB),
+             sqlite3ExprSkipCollateAndLikely(pA),
+             sqlite3ExprSkipCollateAndLikely(pB),
              iTab);
 }
 
+/*
+** Return non-zero if Expr p can only be true if pNN is not NULL.
+**
+** Or if seenNot is true, return non-zero if Expr p can only be
+** non-NULL if pNN is not NULL
+*/
+static int exprImpliesNotNull(
+  Parse *pParse,      /* Parsing context */
+  Expr *p,            /* The expression to be checked */
+  Expr *pNN,          /* The expression that is NOT NULL */
+  int iTab,           /* Table being evaluated */
+  int seenNot         /* Return true only if p can be any non-NULL value */
+){
+  assert( p );
+  assert( pNN );
+  if( sqlite3ExprCompare(pParse, p, pNN, iTab)==0 ){
+    return pNN->op!=TK_NULL;
+  }
+  switch( p->op ){
+    case TK_IN: {
+      if( seenNot && ExprHasProperty(p, EP_xIsSelect) ) return 0;
+      assert( ExprHasProperty(p,EP_xIsSelect)
+           || (p->x.pList!=0 && p->x.pList->nExpr>0) );
+      return exprImpliesNotNull(pParse, p->pLeft, pNN, iTab, 1);
+    }
+    case TK_BETWEEN: {
+      ExprList *pList = p->x.pList;
+      assert( pList!=0 );
+      assert( pList->nExpr==2 );
+      if( seenNot ) return 0;
+      if( exprImpliesNotNull(pParse, pList->a[0].pExpr, pNN, iTab, 1)
+       || exprImpliesNotNull(pParse, pList->a[1].pExpr, pNN, iTab, 1)
+      ){
+        return 1;
+      }
+      return exprImpliesNotNull(pParse, p->pLeft, pNN, iTab, 1);
+    }
+    case TK_EQ:
+    case TK_NE:
+    case TK_LT:
+    case TK_LE:
+    case TK_GT:
+    case TK_GE:
+    case TK_PLUS:
+    case TK_MINUS:
+    case TK_BITOR:
+    case TK_LSHIFT:
+    case TK_RSHIFT: 
+    case TK_CONCAT: 
+      seenNot = 1;
+      /* Fall thru */
+    case TK_STAR:
+    case TK_REM:
+    case TK_BITAND:
+    case TK_SLASH: {
+      if( exprImpliesNotNull(pParse, p->pRight, pNN, iTab, seenNot) ) return 1;
+      /* Fall thru into the next case */
+    }
+    case TK_SPAN:
+    case TK_COLLATE:
+    case TK_UPLUS:
+    case TK_UMINUS: {
+      return exprImpliesNotNull(pParse, p->pLeft, pNN, iTab, seenNot);
+    }
+    case TK_TRUTH: {
+      if( seenNot ) return 0;
+      if( p->op2!=TK_IS ) return 0;
+      return exprImpliesNotNull(pParse, p->pLeft, pNN, iTab, 1);
+    }
+    case TK_BITNOT:
+    case TK_NOT: {
+      return exprImpliesNotNull(pParse, p->pLeft, pNN, iTab, 1);
+    }
+  }
+  return 0;
+}
+
 /*
 ** Return true if we can prove the pE2 will always be true if pE1 is
 ** true.  Return false if we cannot complete the proof or if pE2 might
@@ -101824,16 +104058,16 @@ SQLITE_PRIVATE int sqlite3ExprImpliesExpr(Parse *pParse, Expr *pE1, Expr *pE2, i
   ){
     return 1;
   }
-  if( pE2->op==TK_NOTNULL && pE1->op!=TK_ISNULL && pE1->op!=TK_IS ){
-    Expr *pX = sqlite3ExprSkipCollate(pE1->pLeft);
-    testcase( pX!=pE1->pLeft );
-    if( sqlite3ExprCompare(pParse, pX, pE2->pLeft, iTab)==0 ) return 1;
+  if( pE2->op==TK_NOTNULL
+   && exprImpliesNotNull(pParse, pE1, pE2->pLeft, iTab, 0)
+  ){
+    return 1;
   }
   return 0;
 }
 
 /*
-** This is the Expr node callback for sqlite3ExprImpliesNotNullRow().
+** This is the Expr node callback for sqlite3ExprImpliesNonNullRow().
 ** If the expression node requires that the table at pWalker->iCur
 ** have one or more non-NULL column, then set pWalker->eCode to 1 and abort.
 **
@@ -101847,23 +104081,25 @@ static int impliesNotNullRow(Walker *pWalker, Expr *pExpr){
   if( ExprHasProperty(pExpr, EP_FromJoin) ) return WRC_Prune;
   switch( pExpr->op ){
     case TK_ISNOT:
-    case TK_NOT:
     case TK_ISNULL:
     case TK_NOTNULL:
     case TK_IS:
     case TK_OR:
+    case TK_VECTOR:
     case TK_CASE:
     case TK_IN:
     case TK_FUNCTION:
+    case TK_TRUTH:
       testcase( pExpr->op==TK_ISNOT );
-      testcase( pExpr->op==TK_NOT );
       testcase( pExpr->op==TK_ISNULL );
       testcase( pExpr->op==TK_NOTNULL );
       testcase( pExpr->op==TK_IS );
       testcase( pExpr->op==TK_OR );
+      testcase( pExpr->op==TK_VECTOR );
       testcase( pExpr->op==TK_CASE );
       testcase( pExpr->op==TK_IN );
       testcase( pExpr->op==TK_FUNCTION );
+      testcase( pExpr->op==TK_TRUTH );
       return WRC_Prune;
     case TK_COLUMN:
       if( pWalker->u.iCur==pExpr->iTable ){
@@ -101872,6 +104108,23 @@ static int impliesNotNullRow(Walker *pWalker, Expr *pExpr){
       }
       return WRC_Prune;
 
+    case TK_AND:
+      if( pWalker->eCode==0 ){
+        sqlite3WalkExpr(pWalker, pExpr->pLeft);
+        if( pWalker->eCode ){
+          pWalker->eCode = 0;
+          sqlite3WalkExpr(pWalker, pExpr->pRight);
+        }
+      }
+      return WRC_Prune;
+
+    case TK_BETWEEN:
+      if( sqlite3WalkExpr(pWalker, pExpr->pLeft)==WRC_Abort ){
+        assert( pWalker->eCode );
+        return WRC_Abort;
+      }
+      return WRC_Prune;
+
     /* Virtual tables are allowed to use constraints like x=NULL.  So
     ** a term of the form x=y does not prove that y is not null if x
     ** is the column of a virtual table */
@@ -101892,6 +104145,7 @@ static int impliesNotNullRow(Walker *pWalker, Expr *pExpr){
       ){
        return WRC_Prune;
       }
+
     default:
       return WRC_Continue;
   }
@@ -101921,15 +104175,14 @@ static int impliesNotNullRow(Walker *pWalker, Expr *pExpr){
 */
 SQLITE_PRIVATE int sqlite3ExprImpliesNonNullRow(Expr *p, int iTab){
   Walker w;
-  p = sqlite3ExprSkipCollate(p);
-  while( p ){
-    if( p->op==TK_NOTNULL ){
-      p = p->pLeft;
-    }else if( p->op==TK_AND ){
+  p = sqlite3ExprSkipCollateAndLikely(p);
+  if( p==0 ) return 0;
+  if( p->op==TK_NOTNULL ){
+    p = p->pLeft;
+  }else{
+    while( p->op==TK_AND ){
       if( sqlite3ExprImpliesNonNullRow(p->pLeft, iTab) ) return 1;
       p = p->pRight;
-    }else{
-      break;
     }
   }
   w.xExprCallback = impliesNotNullRow;
@@ -101961,7 +104214,7 @@ struct IdxCover {
 static int exprIdxCover(Walker *pWalker, Expr *pExpr){
   if( pExpr->op==TK_COLUMN
    && pExpr->iTable==pWalker->u.pIdxCover->iCur
-   && sqlite3ColumnOfIndex(pWalker->u.pIdxCover->pIdx, pExpr->iColumn)<0
+   && sqlite3TableColumnToIndex(pWalker->u.pIdxCover->pIdx, pExpr->iColumn)<0
   ){
     pWalker->eCode = 1;
     return WRC_Abort;
@@ -102012,12 +104265,13 @@ struct SrcCount {
 ** Count the number of references to columns.
 */
 static int exprSrcCount(Walker *pWalker, Expr *pExpr){
-  /* The NEVER() on the second term is because sqlite3FunctionUsesThisSrc()
-  ** is always called before sqlite3ExprAnalyzeAggregates() and so the
-  ** TK_COLUMNs have not yet been converted into TK_AGG_COLUMN.  If
-  ** sqlite3FunctionUsesThisSrc() is used differently in the future, the
-  ** NEVER() will need to be removed. */
-  if( pExpr->op==TK_COLUMN || NEVER(pExpr->op==TK_AGG_COLUMN) ){
+  /* There was once a NEVER() on the second term on the grounds that
+  ** sqlite3FunctionUsesThisSrc() was always called before 
+  ** sqlite3ExprAnalyzeAggregates() and so the TK_COLUMNs have not yet 
+  ** been converted into TK_AGG_COLUMN. But this is no longer true due
+  ** to window functions - sqlite3WindowRewrite() may now indirectly call
+  ** FunctionUsesThisSrc() when creating a new sub-select. */
+  if( pExpr->op==TK_COLUMN || pExpr->op==TK_AGG_COLUMN ){
     int i;
     struct SrcCount *p = pWalker->u.pSrcCount;
     SrcList *pSrc = p->pSrc;
@@ -102027,7 +104281,10 @@ static int exprSrcCount(Walker *pWalker, Expr *pExpr){
     }
     if( i<nSrc ){
       p->nThis++;
-    }else{
+    }else if( nSrc==0 || pExpr->iTable<pSrc->a[0].iCursor ){
+      /* In a well-formed parse tree (no name resolution errors),
+      ** TK_COLUMN nodes with smaller Expr.iTable values are in an
+      ** outer context.  Those are the only ones to count as "other" */
       p->nOther++;
     }
   }
@@ -102044,13 +104301,19 @@ SQLITE_PRIVATE int sqlite3FunctionUsesThisSrc(Expr *pExpr, SrcList *pSrcList){
   Walker w;
   struct SrcCount cnt;
   assert( pExpr->op==TK_AGG_FUNCTION );
+  memset(&w, 0, sizeof(w));
   w.xExprCallback = exprSrcCount;
-  w.xSelectCallback = 0;
+  w.xSelectCallback = sqlite3SelectWalkNoop;
   w.u.pSrcCount = &cnt;
   cnt.pSrc = pSrcList;
   cnt.nThis = 0;
   cnt.nOther = 0;
   sqlite3WalkExprList(&w, pExpr->x.pList);
+#ifndef SQLITE_OMIT_WINDOWFUNC
+  if( ExprHasProperty(pExpr, EP_WinFunc) ){
+    sqlite3WalkExpr(&w, pExpr->y.pWin->pFilter);
+  }
+#endif
   return cnt.nThis>0 || cnt.nOther==0;
 }
 
@@ -102279,8 +104542,11 @@ SQLITE_PRIVATE int sqlite3GetTempReg(Parse *pParse){
 ** purpose.
 */
 SQLITE_PRIVATE void sqlite3ReleaseTempReg(Parse *pParse, int iReg){
-  if( iReg && pParse->nTempReg<ArraySize(pParse->aTempReg) ){
-    pParse->aTempReg[pParse->nTempReg++] = iReg;
+  if( iReg ){
+    sqlite3VdbeReleaseRegisters(pParse, iReg, 1, 0, 0);
+    if( pParse->nTempReg<ArraySize(pParse->aTempReg) ){
+      pParse->aTempReg[pParse->nTempReg++] = iReg;
+    }
   }
 }
 
@@ -102306,6 +104572,7 @@ SQLITE_PRIVATE void sqlite3ReleaseTempRange(Parse *pParse, int iReg, int nReg){
     sqlite3ReleaseTempReg(pParse, iReg);
     return;
   }
+  sqlite3VdbeReleaseRegisters(pParse, iReg, nReg, 0, 0);
   if( nReg>pParse->nRangeReg ){
     pParse->nRangeReg = nReg;
     pParse->iRangeReg = iReg;
@@ -102314,6 +104581,11 @@ SQLITE_PRIVATE void sqlite3ReleaseTempRange(Parse *pParse, int iReg, int nReg){
 
 /*
 ** Mark all temporary registers as being unavailable for reuse.
+**
+** Always invoke this procedure after coding a subroutine or co-routine
+** that might be invoked from other parts of the code, to ensure that
+** the sub/co-routine does not use registers in common with the code that
+** invokes the sub/co-routine.
 */
 SQLITE_PRIVATE void sqlite3ClearTempRegCache(Parse *pParse){
   pParse->nTempReg = 0;
@@ -102378,9 +104650,8 @@ SQLITE_PRIVATE int sqlite3NoTempsInRange(Parse *pParse, int iFirst, int iLast){
 static int isAlterableTable(Parse *pParse, Table *pTab){
   if( 0==sqlite3StrNICmp(pTab->zName, "sqlite_", 7) 
 #ifndef SQLITE_OMIT_VIRTUALTABLE
-   || ( (pTab->tabFlags & TF_Shadow) 
-     && (pParse->db->flags & SQLITE_Defensive)
-     && pParse->db->nVdbeExec==0
+   || ( (pTab->tabFlags & TF_Shadow)!=0
+        && sqlite3ReadOnlyShadowTables(pParse->db)
    )
 #endif
   ){
@@ -102401,7 +104672,7 @@ static void renameTestSchema(Parse *pParse, const char *zDb, int bTemp){
   sqlite3NestedParse(pParse, 
       "SELECT 1 "
       "FROM \"%w\".%s "
-      "WHERE name NOT LIKE 'sqlite_%%'"
+      "WHERE name NOT LIKE 'sqliteX_%%' ESCAPE 'X'"
       " AND sql NOT LIKE 'create virtual%%'"
       " AND sqlite_rename_test(%Q, sql, type, name, %d)=NULL ",
       zDb, MASTER_NAME, 
@@ -102412,7 +104683,7 @@ static void renameTestSchema(Parse *pParse, const char *zDb, int bTemp){
     sqlite3NestedParse(pParse, 
         "SELECT 1 "
         "FROM temp.%s "
-        "WHERE name NOT LIKE 'sqlite_%%'"
+        "WHERE name NOT LIKE 'sqliteX_%%' ESCAPE 'X'"
         " AND sql NOT LIKE 'create virtual%%'"
         " AND sqlite_rename_test(%Q, sql, type, name, 1)=NULL ",
         MASTER_NAME, zDb 
@@ -102483,8 +104754,8 @@ SQLITE_PRIVATE void sqlite3AlterRenameTable(
   if( SQLITE_OK!=isAlterableTable(pParse, pTab) ){
     goto exit_rename_table;
   }
-  if( SQLITE_OK!=sqlite3CheckObjectName(pParse, zName) ){ goto
-    exit_rename_table;
+  if( SQLITE_OK!=sqlite3CheckObjectName(pParse,zName,"table",zName) ){
+    goto exit_rename_table;
   }
 
 #ifndef SQLITE_OMIT_VIEW
@@ -102533,7 +104804,7 @@ SQLITE_PRIVATE void sqlite3AlterRenameTable(
       "UPDATE \"%w\".%s SET "
       "sql = sqlite_rename_table(%Q, type, name, sql, %Q, %Q, %d) "
       "WHERE (type!='index' OR tbl_name=%Q COLLATE nocase)"
-      "AND   name NOT LIKE 'sqlite_%%'"
+      "AND   name NOT LIKE 'sqliteX_%%' ESCAPE 'X'"
       , zDb, MASTER_NAME, zDb, zTabName, zName, (iDb==1), zTabName
   );
 
@@ -102544,7 +104815,8 @@ SQLITE_PRIVATE void sqlite3AlterRenameTable(
           "tbl_name = %Q, "
           "name = CASE "
             "WHEN type='table' THEN %Q "
-            "WHEN name LIKE 'sqlite_autoindex%%' AND type='index' THEN "
+            "WHEN name LIKE 'sqliteX_autoindex%%' ESCAPE 'X' "
+            "     AND type='index' THEN "
              "'sqlite_autoindex_' || %Q || substr(name,%d+18) "
             "ELSE name END "
       "WHERE tbl_name=%Q COLLATE nocase AND "
@@ -102644,14 +104916,6 @@ SQLITE_PRIVATE void sqlite3AlterFinishAddColumn(Parse *pParse, Token *pColDef){
   }
 #endif
 
-  /* If the default value for the new column was specified with a 
-  ** literal NULL, then set pDflt to 0. This simplifies checking
-  ** for an SQL NULL default below.
-  */
-  assert( pDflt==0 || pDflt->op==TK_SPAN );
-  if( pDflt && pDflt->pLeft->op==TK_NULL ){
-    pDflt = 0;
-  }
 
   /* Check that the new column is not specified as PRIMARY KEY or UNIQUE.
   ** If there is a NOT NULL constraint, then the default value for the
@@ -102665,36 +104929,50 @@ SQLITE_PRIVATE void sqlite3AlterFinishAddColumn(Parse *pParse, Token *pColDef){
     sqlite3ErrorMsg(pParse, "Cannot add a UNIQUE column");
     return;
   }
-  if( (db->flags&SQLITE_ForeignKeys) && pNew->pFKey && pDflt ){
-    sqlite3ErrorMsg(pParse, 
-        "Cannot add a REFERENCES column with non-NULL default value");
-    return;
-  }
-  if( pCol->notNull && !pDflt ){
-    sqlite3ErrorMsg(pParse, 
-        "Cannot add a NOT NULL column with default value NULL");
-    return;
-  }
-
-  /* Ensure the default expression is something that sqlite3ValueFromExpr()
-  ** can handle (i.e. not CURRENT_TIME etc.)
-  */
-  if( pDflt ){
-    sqlite3_value *pVal = 0;
-    int rc;
-    rc = sqlite3ValueFromExpr(db, pDflt, SQLITE_UTF8, SQLITE_AFF_BLOB, &pVal);
-    assert( rc==SQLITE_OK || rc==SQLITE_NOMEM );
-    if( rc!=SQLITE_OK ){
-      assert( db->mallocFailed == 1 );
+  if( (pCol->colFlags & COLFLAG_GENERATED)==0 ){
+    /* If the default value for the new column was specified with a 
+    ** literal NULL, then set pDflt to 0. This simplifies checking
+    ** for an SQL NULL default below.
+    */
+    assert( pDflt==0 || pDflt->op==TK_SPAN );
+    if( pDflt && pDflt->pLeft->op==TK_NULL ){
+      pDflt = 0;
+    }
+    if( (db->flags&SQLITE_ForeignKeys) && pNew->pFKey && pDflt ){
+      sqlite3ErrorMsg(pParse, 
+          "Cannot add a REFERENCES column with non-NULL default value");
       return;
     }
-    if( !pVal ){
-      sqlite3ErrorMsg(pParse, "Cannot add a column with non-constant default");
+    if( pCol->notNull && !pDflt ){
+      sqlite3ErrorMsg(pParse, 
+          "Cannot add a NOT NULL column with default value NULL");
       return;
     }
-    sqlite3ValueFree(pVal);
+
+    /* Ensure the default expression is something that sqlite3ValueFromExpr()
+    ** can handle (i.e. not CURRENT_TIME etc.)
+    */
+    if( pDflt ){
+      sqlite3_value *pVal = 0;
+      int rc;
+      rc = sqlite3ValueFromExpr(db, pDflt, SQLITE_UTF8, SQLITE_AFF_BLOB, &pVal);
+      assert( rc==SQLITE_OK || rc==SQLITE_NOMEM );
+      if( rc!=SQLITE_OK ){
+        assert( db->mallocFailed == 1 );
+        return;
+      }
+      if( !pVal ){
+        sqlite3ErrorMsg(pParse,"Cannot add a column with non-constant default");
+        return;
+      }
+      sqlite3ValueFree(pVal);
+    }
+  }else if( pCol->colFlags & COLFLAG_STORED ){
+    sqlite3ErrorMsg(pParse, "cannot add a STORED column");
+    return;
   }
 
+
   /* Modify the CREATE TABLE statement. */
   zCol = sqlite3DbStrNDup(db, (char*)pColDef->z, pColDef->n);
   if( zCol ){
@@ -102781,6 +105059,7 @@ SQLITE_PRIVATE void sqlite3AlterBeginAddColumn(Parse *pParse, SrcList *pSrc){
     goto exit_begin_add_column;
   }
 
+  sqlite3MayAbort(pParse);
   assert( pTab->addColOffset>0 );
   iDb = sqlite3SchemaToIndex(db, pTab->pSchema);
 
@@ -102918,7 +105197,8 @@ SQLITE_PRIVATE void sqlite3AlterRenameColumn(
   sqlite3NestedParse(pParse, 
       "UPDATE \"%w\".%s SET "
       "sql = sqlite_rename_column(sql, type, name, %Q, %Q, %d, %Q, %d, %d) "
-      "WHERE name NOT LIKE 'sqlite_%%' AND (type != 'index' OR tbl_name = %Q)"
+      "WHERE name NOT LIKE 'sqliteX_%%' ESCAPE 'X' "
+      " AND (type != 'index' OR tbl_name = %Q)"
       " AND sql NOT LIKE 'create virtual%%'",
       zDb, MASTER_NAME, 
       zDb, pTab->zName, iCol, zNew, bQuote, iSchema==1,
@@ -103036,12 +105316,14 @@ SQLITE_PRIVATE void *sqlite3RenameTokenMap(Parse *pParse, void *pPtr, Token *pTo
   RenameToken *pNew;
   assert( pPtr || pParse->db->mallocFailed );
   renameTokenCheckAll(pParse, pPtr);
-  pNew = sqlite3DbMallocZero(pParse->db, sizeof(RenameToken));
-  if( pNew ){
-    pNew->p = pPtr;
-    pNew->t = *pToken;
-    pNew->pNext = pParse->pRename;
-    pParse->pRename = pNew;
+  if( pParse->eParseMode!=PARSE_MODE_UNMAP ){
+    pNew = sqlite3DbMallocZero(pParse->db, sizeof(RenameToken));
+    if( pNew ){
+      pNew->p = pPtr;
+      pNew->t = *pToken;
+      pNew->pNext = pParse->pRename;
+      pParse->pRename = pNew;
+    }
   }
 
   return pPtr;
@@ -103072,15 +105354,67 @@ static int renameUnmapExprCb(Walker *pWalker, Expr *pExpr){
   return WRC_Continue;
 }
 
+/*
+** Iterate through the Select objects that are part of WITH clauses attached
+** to select statement pSelect.
+*/
+static void renameWalkWith(Walker *pWalker, Select *pSelect){
+  With *pWith = pSelect->pWith;
+  if( pWith ){
+    int i;
+    for(i=0; i<pWith->nCte; i++){
+      Select *p = pWith->a[i].pSelect;
+      NameContext sNC;
+      memset(&sNC, 0, sizeof(sNC));
+      sNC.pParse = pWalker->pParse;
+      sqlite3SelectPrep(sNC.pParse, p, &sNC);
+      sqlite3WalkSelect(pWalker, p);
+      sqlite3RenameExprlistUnmap(pWalker->pParse, pWith->a[i].pCols);
+    }
+  }
+}
+
+/*
+** Walker callback used by sqlite3RenameExprUnmap().
+*/
+static int renameUnmapSelectCb(Walker *pWalker, Select *p){
+  Parse *pParse = pWalker->pParse;
+  int i;
+  if( pParse->nErr ) return WRC_Abort;
+  if( NEVER(p->selFlags & SF_View) ) return WRC_Prune;
+  if( ALWAYS(p->pEList) ){
+    ExprList *pList = p->pEList;
+    for(i=0; i<pList->nExpr; i++){
+      if( pList->a[i].zEName && pList->a[i].eEName==ENAME_NAME ){
+        sqlite3RenameTokenRemap(pParse, 0, (void*)pList->a[i].zEName);
+      }
+    }
+  }
+  if( ALWAYS(p->pSrc) ){  /* Every Select as a SrcList, even if it is empty */
+    SrcList *pSrc = p->pSrc;
+    for(i=0; i<pSrc->nSrc; i++){
+      sqlite3RenameTokenRemap(pParse, 0, (void*)pSrc->a[i].zName);
+      if( sqlite3WalkExpr(pWalker, pSrc->a[i].pOn) ) return WRC_Abort;
+    }
+  }
+
+  renameWalkWith(pWalker, p);
+  return WRC_Continue;
+}
+
 /*
 ** Remove all nodes that are part of expression pExpr from the rename list.
 */
 SQLITE_PRIVATE void sqlite3RenameExprUnmap(Parse *pParse, Expr *pExpr){
+  u8 eMode = pParse->eParseMode;
   Walker sWalker;
   memset(&sWalker, 0, sizeof(Walker));
   sWalker.pParse = pParse;
   sWalker.xExprCallback = renameUnmapExprCb;
+  sWalker.xSelectCallback = renameUnmapSelectCb;
+  pParse->eParseMode = PARSE_MODE_UNMAP;
   sqlite3WalkExpr(&sWalker, pExpr);
+  pParse->eParseMode = eMode;
 }
 
 /*
@@ -103096,7 +105430,9 @@ SQLITE_PRIVATE void sqlite3RenameExprlistUnmap(Parse *pParse, ExprList *pEList){
     sWalker.xExprCallback = renameUnmapExprCb;
     sqlite3WalkExprList(&sWalker, pEList);
     for(i=0; i<pEList->nExpr; i++){
-      sqlite3RenameTokenRemap(pParse, 0, (void*)pEList->a[i].zName);
+      if( ALWAYS(pEList->a[i].eEName==ENAME_NAME) ){
+        sqlite3RenameTokenRemap(pParse, 0, (void*)pEList->a[i].zEName);
+      }
     }
   }
 }
@@ -103134,30 +105470,13 @@ static void renameTokenFind(Parse *pParse, struct RenameCtx *pCtx, void *pPtr){
   }
 }
 
-/*
-** Iterate through the Select objects that are part of WITH clauses attached
-** to select statement pSelect.
-*/
-static void renameWalkWith(Walker *pWalker, Select *pSelect){
-  if( pSelect->pWith ){
-    int i;
-    for(i=0; i<pSelect->pWith->nCte; i++){
-      Select *p = pSelect->pWith->a[i].pSelect;
-      NameContext sNC;
-      memset(&sNC, 0, sizeof(sNC));
-      sNC.pParse = pWalker->pParse;
-      sqlite3SelectPrep(sNC.pParse, p, &sNC);
-      sqlite3WalkSelect(pWalker, p);
-    }
-  }
-}
-
 /*
 ** This is a Walker select callback. It does nothing. It is only required
 ** because without a dummy callback, sqlite3WalkExpr() and similar do not
 ** descend into sub-select statements.
 */
 static int renameColumnSelectCb(Walker *pWalker, Select *p){
+  if( p->selFlags & SF_View ) return WRC_Prune;
   renameWalkWith(pWalker, p);
   return WRC_Continue;
 }
@@ -103251,8 +105570,11 @@ static void renameColumnElistNames(
   if( pEList ){
     int i;
     for(i=0; i<pEList->nExpr; i++){
-      char *zName = pEList->a[i].zName;
-      if( 0==sqlite3_stricmp(zName, zOld) ){
+      char *zName = pEList->a[i].zEName;
+      if( ALWAYS(pEList->a[i].eEName==ENAME_NAME)
+       && ALWAYS(zName!=0)
+       && 0==sqlite3_stricmp(zName, zOld)
+      ){
         renameTokenFind(pParse, pCtx, (void*)zName);
       }
     }
@@ -103288,7 +105610,6 @@ static void renameColumnIdlistNames(
 static int renameParseSql(
   Parse *p,                       /* Memory to use for Parse object */
   const char *zDb,                /* Name of schema SQL belongs to */
-  int bTable,                     /* 1 -> RENAME TABLE, 0 -> RENAME COLUMN */
   sqlite3 *db,                    /* Database handle */
   const char *zSql,               /* SQL to parse */
   int bTemp                       /* True if SQL is from temp schema */
@@ -103302,7 +105623,7 @@ static int renameParseSql(
   ** occurs and the parse does not result in a new table, index or
   ** trigger object, the database must be corrupt. */
   memset(p, 0, sizeof(Parse));
-  p->eParseMode = (bTable ? PARSE_MODE_RENAME_TABLE : PARSE_MODE_RENAME_COLUMN);
+  p->eParseMode = PARSE_MODE_RENAME;
   p->db = db;
   p->nQueryLoop = 1;
   rc = sqlite3RunParser(p, zSql, &zErr);
@@ -103377,7 +105698,7 @@ static int renameEditSql(
   ** with the new column name. All that remains is to construct and
   ** return the edited SQL string. */
   assert( nQuot>=nNew );
-  zOut = sqlite3DbMallocZero(db, (u64)pRename->nList*nQuot + 1 + nSql);
+  zOut = sqlite3DbMallocZero(db, nSql + pRename->nList*nQuot + 1);
   if( zOut ){
     int nOut = nSql;
     memcpy(zOut, zSql, nSql);
@@ -103609,7 +105930,7 @@ static void renameColumnFunc(
 #ifndef SQLITE_OMIT_AUTHORIZATION
   db->xAuth = 0;
 #endif
-  rc = renameParseSql(&sParse, zDb, 0, db, zSql, bTemp);
+  rc = renameParseSql(&sParse, zDb, db, zSql, bTemp);
 
   /* Find tokens that need to be replaced. */
   memset(&sWalker, 0, sizeof(Walker));
@@ -103623,8 +105944,9 @@ static void renameColumnFunc(
   if( sParse.pNewTable ){
     Select *pSelect = sParse.pNewTable->pSelect;
     if( pSelect ){
+      pSelect->selFlags &= ~SF_View;
       sParse.rc = SQLITE_OK;
-      sqlite3SelectPrep(&sParse, sParse.pNewTable->pSelect, 0);
+      sqlite3SelectPrep(&sParse, pSelect, 0);
       rc = (db->mallocFailed ? SQLITE_NOMEM : sParse.rc);
       if( rc==SQLITE_OK ){
         sqlite3WalkSelect(&sWalker, pSelect);
@@ -103651,6 +105973,11 @@ static void renameColumnFunc(
           sqlite3WalkExprList(&sWalker, pIdx->aColExpr);
         }
       }
+#ifndef SQLITE_OMIT_GENERATED_COLUMNS
+      for(i=0; i<sParse.pNewTable->nCol; i++){
+        sqlite3WalkExpr(&sWalker, sParse.pNewTable->aCol[i].pDflt);
+      }
+#endif
 
       for(pFKey=sParse.pNewTable->pFKey; pFKey; pFKey=pFKey->pNextFrom){
         for(i=0; i<pFKey->nCol; i++){
@@ -103736,6 +106063,7 @@ static int renameTableSelectCb(Walker *pWalker, Select *pSelect){
   int i;
   RenameCtx *p = pWalker->u.pRename;
   SrcList *pSrc = pSelect->pSrc;
+  if( pSelect->selFlags & SF_View ) return WRC_Prune;
   if( pSrc==0 ){
     assert( pWalker->pParse->db->mallocFailed );
     return WRC_Abort;
@@ -103806,7 +106134,7 @@ static void renameTableFunc(
     sWalker.xSelectCallback = renameTableSelectCb;
     sWalker.u.pRename = &sCtx;
 
-    rc = renameParseSql(&sParse, zDb, 1, db, zInput, bTemp);
+    rc = renameParseSql(&sParse, zDb, db, zInput, bTemp);
 
     if( rc==SQLITE_OK ){
       int isLegacy = (db->flags & SQLITE_LegacyAlter);
@@ -103815,13 +106143,19 @@ static void renameTableFunc(
 
         if( pTab->pSelect ){
           if( isLegacy==0 ){
+            Select *pSelect = pTab->pSelect;
             NameContext sNC;
             memset(&sNC, 0, sizeof(sNC));
             sNC.pParse = &sParse;
 
+            assert( pSelect->selFlags & SF_View );
+            pSelect->selFlags &= ~SF_View;
             sqlite3SelectPrep(&sParse, pTab->pSelect, &sNC);
-            if( sParse.nErr ) rc = sParse.rc;
-            sqlite3WalkSelect(&sWalker, pTab->pSelect);
+            if( sParse.nErr ){
+              rc = sParse.rc;
+            }else{
+              sqlite3WalkSelect(&sWalker, pTab->pSelect);
+            }
           }
         }else{
           /* Modify any FK definitions to point to the new table. */
@@ -103942,7 +106276,7 @@ static void renameTableTest(
   if( zDb && zInput ){
     int rc;
     Parse sParse;
-    rc = renameParseSql(&sParse, zDb, 1, db, zInput, bTemp);
+    rc = renameParseSql(&sParse, zDb, db, zInput, bTemp);
     if( rc==SQLITE_OK ){
       if( isLegacy==0 && sParse.pNewTable && sParse.pNewTable->pSelect ){
         NameContext sNC;
@@ -104019,13 +106353,13 @@ SQLITE_PRIVATE void sqlite3AlterFunctions(void){
 ** is between 3.6.18 and 3.7.8, inclusive, and unless SQLite is compiled
 ** with SQLITE_ENABLE_STAT2.  The sqlite_stat2 table is deprecated.
 ** The sqlite_stat2 table is superseded by sqlite_stat3, which is only
-** created and used by SQLite versions 3.7.9 and later and with
+** created and used by SQLite versions 3.7.9 through 3.29.0 when
 ** SQLITE_ENABLE_STAT3 defined.  The functionality of sqlite_stat3
-** is a superset of sqlite_stat2.  The sqlite_stat4 is an enhanced
-** version of sqlite_stat3 and is only available when compiled with
-** SQLITE_ENABLE_STAT4 and in SQLite versions 3.8.1 and later.  It is
-** not possible to enable both STAT3 and STAT4 at the same time.  If they
-** are both enabled, then STAT4 takes precedence.
+** is a superset of sqlite_stat2 and is also now deprecated.  The
+** sqlite_stat4 is an enhanced version of sqlite_stat3 and is only 
+** available when compiled with SQLITE_ENABLE_STAT4 and in SQLite
+** versions 3.8.1 and later.  STAT4 is the only variant that is still
+** supported.
 **
 ** For most applications, sqlite_stat1 provides all the statistics required
 ** for the query planner to make good choices.
@@ -104136,17 +106470,11 @@ SQLITE_PRIVATE void sqlite3AlterFunctions(void){
 
 #if defined(SQLITE_ENABLE_STAT4)
 # define IsStat4     1
-# define IsStat3     0
-#elif defined(SQLITE_ENABLE_STAT3)
-# define IsStat4     0
-# define IsStat3     1
 #else
 # define IsStat4     0
-# define IsStat3     0
 # undef SQLITE_STAT4_SAMPLES
 # define SQLITE_STAT4_SAMPLES 1
 #endif
-#define IsStat34    (IsStat3+IsStat4)  /* 1 for STAT3 or STAT4. 0 otherwise */
 
 /*
 ** This routine generates code that opens the sqlite_statN tables.
@@ -104175,14 +106503,10 @@ static void openStatTable(
     { "sqlite_stat1", "tbl,idx,stat" },
 #if defined(SQLITE_ENABLE_STAT4)
     { "sqlite_stat4", "tbl,idx,neq,nlt,ndlt,sample" },
-    { "sqlite_stat3", 0 },
-#elif defined(SQLITE_ENABLE_STAT3)
-    { "sqlite_stat3", "tbl,idx,neq,nlt,ndlt,sample" },
-    { "sqlite_stat4", 0 },
 #else
-    { "sqlite_stat3", 0 },
     { "sqlite_stat4", 0 },
 #endif
+    { "sqlite_stat3", 0 },
   };
   int i;
   sqlite3 *db = pParse->db;
@@ -104263,7 +106587,7 @@ typedef struct Stat4Sample Stat4Sample;
 struct Stat4Sample {
   tRowcnt *anEq;                  /* sqlite_stat4.nEq */
   tRowcnt *anDLt;                 /* sqlite_stat4.nDLt */
-#ifdef SQLITE_ENABLE_STAT3_OR_STAT4
+#ifdef SQLITE_ENABLE_STAT4
   tRowcnt *anLt;                  /* sqlite_stat4.nLt */
   union {
     i64 iRowid;                     /* Rowid in main table of the key */
@@ -104294,7 +106618,7 @@ struct Stat4Accum {
 
 /* Reclaim memory used by a Stat4Sample
 */
-#ifdef SQLITE_ENABLE_STAT3_OR_STAT4
+#ifdef SQLITE_ENABLE_STAT4
 static void sampleClear(sqlite3 *db, Stat4Sample *p){
   assert( db!=0 );
   if( p->nRowid ){
@@ -104306,7 +106630,7 @@ static void sampleClear(sqlite3 *db, Stat4Sample *p){
 
 /* Initialize the BLOB value of a ROWID
 */
-#ifdef SQLITE_ENABLE_STAT3_OR_STAT4
+#ifdef SQLITE_ENABLE_STAT4
 static void sampleSetRowid(sqlite3 *db, Stat4Sample *p, int n, const u8 *pData){
   assert( db!=0 );
   if( p->nRowid ) sqlite3DbFree(db, p->u.aRowid);
@@ -104322,7 +106646,7 @@ static void sampleSetRowid(sqlite3 *db, Stat4Sample *p, int n, const u8 *pData){
 
 /* Initialize the INTEGER value of a ROWID.
 */
-#ifdef SQLITE_ENABLE_STAT3_OR_STAT4
+#ifdef SQLITE_ENABLE_STAT4
 static void sampleSetRowidInt64(sqlite3 *db, Stat4Sample *p, i64 iRowid){
   assert( db!=0 );
   if( p->nRowid ) sqlite3DbFree(db, p->u.aRowid);
@@ -104335,7 +106659,7 @@ static void sampleSetRowidInt64(sqlite3 *db, Stat4Sample *p, i64 iRowid){
 /*
 ** Copy the contents of object (*pFrom) into (*pTo).
 */
-#ifdef SQLITE_ENABLE_STAT3_OR_STAT4
+#ifdef SQLITE_ENABLE_STAT4
 static void sampleCopy(Stat4Accum *p, Stat4Sample *pTo, Stat4Sample *pFrom){
   pTo->isPSample = pFrom->isPSample;
   pTo->iCol = pFrom->iCol;
@@ -104356,7 +106680,7 @@ static void sampleCopy(Stat4Accum *p, Stat4Sample *pTo, Stat4Sample *pFrom){
 */
 static void stat4Destructor(void *pOld){
   Stat4Accum *p = (Stat4Accum*)pOld;
-#ifdef SQLITE_ENABLE_STAT3_OR_STAT4
+#ifdef SQLITE_ENABLE_STAT4
   int i;
   for(i=0; i<p->nCol; i++) sampleClear(p->db, p->aBest+i);
   for(i=0; i<p->mxSample; i++) sampleClear(p->db, p->a+i);
@@ -104376,7 +106700,7 @@ static void stat4Destructor(void *pOld){
 ** WITHOUT ROWID table, N is the number of PRIMARY KEY columns, not the
 ** total number of columns in the table.
 **
-** Note 2:  C is only used for STAT3 and STAT4.
+** Note 2:  C is only used for STAT4.
 **
 ** For indexes on ordinary rowid tables, N==K+1.  But for indexes on
 ** WITHOUT ROWID tables, N=K+P where P is the number of columns in the
@@ -104399,7 +106723,7 @@ static void statInit(
   int nColUp;                     /* nCol rounded up for alignment */
   int n;                          /* Bytes of space to allocate */
   sqlite3 *db;                    /* Database connection */
-#ifdef SQLITE_ENABLE_STAT3_OR_STAT4
+#ifdef SQLITE_ENABLE_STAT4
   int mxSample = SQLITE_STAT4_SAMPLES;
 #endif
 
@@ -104416,7 +106740,7 @@ static void statInit(
   n = sizeof(*p) 
     + sizeof(tRowcnt)*nColUp                  /* Stat4Accum.anEq */
     + sizeof(tRowcnt)*nColUp                  /* Stat4Accum.anDLt */
-#ifdef SQLITE_ENABLE_STAT3_OR_STAT4
+#ifdef SQLITE_ENABLE_STAT4
     + sizeof(tRowcnt)*nColUp                  /* Stat4Accum.anLt */
     + sizeof(Stat4Sample)*(nCol+mxSample)     /* Stat4Accum.aBest[], a[] */
     + sizeof(tRowcnt)*3*nColUp*(nCol+mxSample)
@@ -104436,7 +106760,7 @@ static void statInit(
   p->current.anDLt = (tRowcnt*)&p[1];
   p->current.anEq = &p->current.anDLt[nColUp];
 
-#ifdef SQLITE_ENABLE_STAT3_OR_STAT4
+#ifdef SQLITE_ENABLE_STAT4
   {
     u8 *pSpace;                     /* Allocated space not yet assigned */
     int i;                          /* Used to iterate through p->aSample[] */
@@ -104471,7 +106795,7 @@ static void statInit(
   sqlite3_result_blob(context, p, sizeof(*p), stat4Destructor);
 }
 static const FuncDef statInitFuncdef = {
-  2+IsStat34,      /* nArg */
+  2+IsStat4,       /* nArg */
   SQLITE_UTF8,     /* funcFlags */
   0,               /* pUserData */
   0,               /* pNext */
@@ -104511,7 +106835,7 @@ static int sampleIsBetterPost(
 }
 #endif
 
-#ifdef SQLITE_ENABLE_STAT3_OR_STAT4
+#ifdef SQLITE_ENABLE_STAT4
 /*
 ** Return true if pNew is to be preferred over pOld.
 **
@@ -104530,15 +106854,11 @@ static int sampleIsBetter(
   assert( IsStat4 || (pNew->iCol==0 && pOld->iCol==0) );
 
   if( (nEqNew>nEqOld) ) return 1;
-#ifdef SQLITE_ENABLE_STAT4
   if( nEqNew==nEqOld ){
     if( pNew->iCol<pOld->iCol ) return 1;
     return (pNew->iCol==pOld->iCol && sampleIsBetterPost(pAccum, pNew, pOld));
   }
   return 0;
-#else
-  return (nEqNew==nEqOld && pNew->iHash>pOld->iHash);
-#endif
 }
 
 /*
@@ -104551,7 +106871,6 @@ static void sampleInsert(Stat4Accum *p, Stat4Sample *pNew, int nEqZero){
 
   assert( IsStat4 || nEqZero==0 );
 
-#ifdef SQLITE_ENABLE_STAT4
   /* Stat4Accum.nMaxEqZero is set to the maximum number of leading 0
   ** values in the anEq[] array of any sample in Stat4Accum.a[]. In
   ** other words, if nMaxEqZero is n, then it is guaranteed that there
@@ -104585,7 +106904,6 @@ static void sampleInsert(Stat4Accum *p, Stat4Sample *pNew, int nEqZero){
       goto find_new_min;
     }
   }
-#endif
 
   /* If necessary, remove sample iMin to make room for the new sample. */
   if( p->nSample>=p->mxSample ){
@@ -104606,10 +106924,8 @@ static void sampleInsert(Stat4Accum *p, Stat4Sample *pNew, int nEqZero){
   /* The "rows less-than" for the rowid column must be greater than that
   ** for the last sample in the p->a[] array. Otherwise, the samples would
   ** be out of order. */
-#ifdef SQLITE_ENABLE_STAT4
   assert( p->nSample==0 
        || pNew->anLt[p->nCol-1] > p->a[p->nSample-1].anLt[p->nCol-1] );
-#endif
 
   /* Insert the new sample */
   pSample = &p->a[p->nSample];
@@ -104619,9 +106935,7 @@ static void sampleInsert(Stat4Accum *p, Stat4Sample *pNew, int nEqZero){
   /* Zero the first nEqZero entries in the anEq[] array. */
   memset(pSample->anEq, 0, sizeof(tRowcnt)*nEqZero);
 
-#ifdef SQLITE_ENABLE_STAT4
- find_new_min:
-#endif
+find_new_min:
   if( p->nSample>=p->mxSample ){
     int iMin = -1;
     for(i=0; i<p->mxSample; i++){
@@ -104634,7 +106948,7 @@ static void sampleInsert(Stat4Accum *p, Stat4Sample *pNew, int nEqZero){
     p->iMin = iMin;
   }
 }
-#endif /* SQLITE_ENABLE_STAT3_OR_STAT4 */
+#endif /* SQLITE_ENABLE_STAT4 */
 
 /*
 ** Field iChng of the index being scanned has changed. So at this point
@@ -104675,28 +106989,7 @@ static void samplePushPrevious(Stat4Accum *p, int iChng){
   }
 #endif
 
-#if defined(SQLITE_ENABLE_STAT3) && !defined(SQLITE_ENABLE_STAT4)
-  if( iChng==0 ){
-    tRowcnt nLt = p->current.anLt[0];
-    tRowcnt nEq = p->current.anEq[0];
-
-    /* Check if this is to be a periodic sample. If so, add it. */
-    if( (nLt/p->nPSample)!=(nLt+nEq)/p->nPSample ){
-      p->current.isPSample = 1;
-      sampleInsert(p, &p->current, 0);
-      p->current.isPSample = 0;
-    }else 
-
-    /* Or if it is a non-periodic sample. Add it in this case too. */
-    if( p->nSample<p->mxSample 
-     || sampleIsBetter(p, &p->current, &p->a[p->iMin]) 
-    ){
-      sampleInsert(p, &p->current, 0);
-    }
-  }
-#endif
-
-#ifndef SQLITE_ENABLE_STAT3_OR_STAT4
+#ifndef SQLITE_ENABLE_STAT4
   UNUSED_PARAMETER( p );
   UNUSED_PARAMETER( iChng );
 #endif
@@ -104716,7 +107009,7 @@ static void samplePushPrevious(Stat4Accum *p, int iChng){
 ** index being analyzed.  The stat_get() SQL function will later be used to
 ** extract relevant information for constructing the sqlite_statN tables.
 **
-** The R parameter is only used for STAT3 and STAT4
+** The R parameter is only used for STAT4
 */
 static void statPush(
   sqlite3_context *context,
@@ -104748,14 +107041,14 @@ static void statPush(
     }
     for(i=iChng; i<p->nCol; i++){
       p->current.anDLt[i]++;
-#ifdef SQLITE_ENABLE_STAT3_OR_STAT4
+#ifdef SQLITE_ENABLE_STAT4
       p->current.anLt[i] += p->current.anEq[i];
 #endif
       p->current.anEq[i] = 1;
     }
   }
   p->nRow++;
-#ifdef SQLITE_ENABLE_STAT3_OR_STAT4
+#ifdef SQLITE_ENABLE_STAT4
   if( sqlite3_value_type(argv[2])==SQLITE_INTEGER ){
     sampleSetRowidInt64(p->db, &p->current, sqlite3_value_int64(argv[2]));
   }else{
@@ -104788,7 +107081,7 @@ static void statPush(
 #endif
 }
 static const FuncDef statPushFuncdef = {
-  2+IsStat34,      /* nArg */
+  2+IsStat4,       /* nArg */
   SQLITE_UTF8,     /* funcFlags */
   0,               /* pUserData */
   0,               /* pNext */
@@ -104819,7 +107112,7 @@ static const FuncDef statPushFuncdef = {
 ** parameter will always be a poiner to a Stat4Accum object, never a
 ** NULL.
 **
-** If neither STAT3 nor STAT4 are enabled, then J is always
+** If STAT4 is not enabled, then J is always
 ** STAT_GET_STAT1 and is hence omitted and this routine becomes
 ** a one-parameter function, stat_get(P), that always returns the
 ** stat1 table entry information.
@@ -104830,8 +107123,8 @@ static void statGet(
   sqlite3_value **argv
 ){
   Stat4Accum *p = (Stat4Accum*)sqlite3_value_blob(argv[0]);
-#ifdef SQLITE_ENABLE_STAT3_OR_STAT4
-  /* STAT3 and STAT4 have a parameter on this routine. */
+#ifdef SQLITE_ENABLE_STAT4
+  /* STAT4 has a parameter on this routine. */
   int eCall = sqlite3_value_int(argv[1]);
   assert( argc==2 );
   assert( eCall==STAT_GET_STAT1 || eCall==STAT_GET_NEQ 
@@ -104867,7 +107160,7 @@ static void statGet(
     char *z;
     int i;
 
-    char *zRet = sqlite3MallocZero( ((u64)p->nKeyCol+1)*25 );
+    char *zRet = sqlite3MallocZero( (p->nKeyCol+1)*25 );
     if( zRet==0 ){
       sqlite3_result_error_nomem(context);
       return;
@@ -104886,7 +107179,7 @@ static void statGet(
 
     sqlite3_result_text(context, zRet, -1, sqlite3_free);
   }
-#ifdef SQLITE_ENABLE_STAT3_OR_STAT4
+#ifdef SQLITE_ENABLE_STAT4
   else if( eCall==STAT_GET_ROWID ){
     if( p->iGet<0 ){
       samplePushPrevious(p, 0);
@@ -104915,9 +107208,7 @@ static void statGet(
       }
     }
 
-    if( IsStat3 ){
-      sqlite3_result_int64(context, (i64)aCnt[0]);
-    }else{
+    {
       char *zRet = sqlite3MallocZero(p->nCol * 25);
       if( zRet==0 ){
         sqlite3_result_error_nomem(context);
@@ -104934,13 +107225,13 @@ static void statGet(
       }
     }
   }
-#endif /* SQLITE_ENABLE_STAT3_OR_STAT4 */
+#endif /* SQLITE_ENABLE_STAT4 */
 #ifndef SQLITE_DEBUG
   UNUSED_PARAMETER( argc );
 #endif
 }
 static const FuncDef statGetFuncdef = {
-  1+IsStat34,      /* nArg */
+  1+IsStat4,       /* nArg */
   SQLITE_UTF8,     /* funcFlags */
   0,               /* pUserData */
   0,               /* pNext */
@@ -104951,18 +107242,17 @@ static const FuncDef statGetFuncdef = {
   {0}
 };
 
-static void callStatGet(Vdbe *v, int regStat4, int iParam, int regOut){
-  assert( regOut!=regStat4 && regOut!=regStat4+1 );
-#ifdef SQLITE_ENABLE_STAT3_OR_STAT4
-  sqlite3VdbeAddOp2(v, OP_Integer, iParam, regStat4+1);
+static void callStatGet(Parse *pParse, int regStat4, int iParam, int regOut){
+#ifdef SQLITE_ENABLE_STAT4
+  sqlite3VdbeAddOp2(pParse->pVdbe, OP_Integer, iParam, regStat4+1);
 #elif SQLITE_DEBUG
   assert( iParam==STAT_GET_STAT1 );
 #else
   UNUSED_PARAMETER( iParam );
 #endif
-  sqlite3VdbeAddOp4(v, OP_Function0, 0, regStat4, regOut,
-                    (char*)&statGetFuncdef, P4_FUNCDEF);
-  sqlite3VdbeChangeP5(v, 1 + IsStat34);
+  assert( regOut!=regStat4 && regOut!=regStat4+1 );
+  sqlite3VdbeAddFunctionCall(pParse, 0, regStat4, regOut, 1+IsStat4,
+                             &statGetFuncdef, 0);
 }
 
 /*
@@ -104989,7 +107279,7 @@ static void analyzeOneTable(
   int regNewRowid = iMem++;    /* Rowid for the inserted record */
   int regStat4 = iMem++;       /* Register to hold Stat4Accum object */
   int regChng = iMem++;        /* Index of changed index field */
-#ifdef SQLITE_ENABLE_STAT3_OR_STAT4
+#ifdef SQLITE_ENABLE_STAT4
   int regRowid = iMem++;       /* Rowid argument passed to stat_push() */
 #endif
   int regTemp = iMem++;        /* Temporary use register */
@@ -105123,16 +107413,15 @@ static void analyzeOneTable(
     **    (3) the number of rows in the index,
     **
     **
-    ** The third argument is only used for STAT3 and STAT4
+    ** The third argument is only used for STAT4
     */
-#ifdef SQLITE_ENABLE_STAT3_OR_STAT4
+#ifdef SQLITE_ENABLE_STAT4
     sqlite3VdbeAddOp2(v, OP_Count, iIdxCur, regStat4+3);
 #endif
     sqlite3VdbeAddOp2(v, OP_Integer, nCol, regStat4+1);
     sqlite3VdbeAddOp2(v, OP_Integer, pIdx->nKeyCol, regStat4+2);
-    sqlite3VdbeAddOp4(v, OP_Function0, 0, regStat4+1, regStat4,
-                     (char*)&statInitFuncdef, P4_FUNCDEF);
-    sqlite3VdbeChangeP5(v, 2+IsStat34);
+    sqlite3VdbeAddFunctionCall(pParse, 0, regStat4+1, regStat4, 2+IsStat4,
+                               &statInitFuncdef, 0);
 
     /* Implementation of the following:
     **
@@ -105203,12 +107492,12 @@ static void analyzeOneTable(
   
     /*
     **  chng_addr_N:
-    **   regRowid = idx(rowid)            // STAT34 only
-    **   stat_push(P, regChng, regRowid)  // 3rd parameter STAT34 only
+    **   regRowid = idx(rowid)            // STAT4 only
+    **   stat_push(P, regChng, regRowid)  // 3rd parameter STAT4 only
     **   Next csr
     **   if !eof(csr) goto next_row;
     */
-#ifdef SQLITE_ENABLE_STAT3_OR_STAT4
+#ifdef SQLITE_ENABLE_STAT4
     assert( regRowid==(regStat4+2) );
     if( HasRowid(pTab) ){
       sqlite3VdbeAddOp2(v, OP_IdxRowid, iIdxCur, regRowid);
@@ -105217,7 +107506,7 @@ static void analyzeOneTable(
       int j, k, regKey;
       regKey = sqlite3GetTempRange(pParse, pPk->nKeyCol);
       for(j=0; j<pPk->nKeyCol; j++){
-        k = sqlite3ColumnOfIndex(pIdx, pPk->aiColumn[j]);
+        k = sqlite3TableColumnToIndex(pIdx, pPk->aiColumn[j]);
         assert( k>=0 && k<pIdx->nColumn );
         sqlite3VdbeAddOp3(v, OP_Column, iIdxCur, k, regKey+j);
         VdbeComment((v, "%s", pTab->aCol[pPk->aiColumn[j]].zName));
@@ -105227,13 +107516,12 @@ static void analyzeOneTable(
     }
 #endif
     assert( regChng==(regStat4+1) );
-    sqlite3VdbeAddOp4(v, OP_Function0, 1, regStat4, regTemp,
-                     (char*)&statPushFuncdef, P4_FUNCDEF);
-    sqlite3VdbeChangeP5(v, 2+IsStat34);
+    sqlite3VdbeAddFunctionCall(pParse, 1, regStat4, regTemp, 2+IsStat4,
+                               &statPushFuncdef, 0);
     sqlite3VdbeAddOp2(v, OP_Next, iIdxCur, addrNextRow); VdbeCoverage(v);
 
     /* Add the entry to the stat1 table. */
-    callStatGet(v, regStat4, STAT_GET_STAT1, regStat1);
+    callStatGet(pParse, regStat4, STAT_GET_STAT1, regStat1);
     assert( "BBB"[0]==SQLITE_AFF_TEXT );
     sqlite3VdbeAddOp4(v, OP_MakeRecord, regTabname, 3, regTemp, "BBB", 0);
     sqlite3VdbeAddOp2(v, OP_NewRowid, iStatCur, regNewRowid);
@@ -105243,8 +107531,8 @@ static void analyzeOneTable(
 #endif
     sqlite3VdbeChangeP5(v, OPFLAG_APPEND);
 
-    /* Add the entries to the stat3 or stat4 table. */
-#ifdef SQLITE_ENABLE_STAT3_OR_STAT4
+    /* Add the entries to the stat4 table. */
+#ifdef SQLITE_ENABLE_STAT4
     {
       int regEq = regStat1;
       int regLt = regStat1+1;
@@ -105259,29 +107547,25 @@ static void analyzeOneTable(
       pParse->nMem = MAX(pParse->nMem, regCol+nCol);
 
       addrNext = sqlite3VdbeCurrentAddr(v);
-      callStatGet(v, regStat4, STAT_GET_ROWID, regSampleRowid);
+      callStatGet(pParse, regStat4, STAT_GET_ROWID, regSampleRowid);
       addrIsNull = sqlite3VdbeAddOp1(v, OP_IsNull, regSampleRowid);
       VdbeCoverage(v);
-      callStatGet(v, regStat4, STAT_GET_NEQ, regEq);
-      callStatGet(v, regStat4, STAT_GET_NLT, regLt);
-      callStatGet(v, regStat4, STAT_GET_NDLT, regDLt);
+      callStatGet(pParse, regStat4, STAT_GET_NEQ, regEq);
+      callStatGet(pParse, regStat4, STAT_GET_NLT, regLt);
+      callStatGet(pParse, regStat4, STAT_GET_NDLT, regDLt);
       sqlite3VdbeAddOp4Int(v, seekOp, iTabCur, addrNext, regSampleRowid, 0);
       VdbeCoverage(v);
-#ifdef SQLITE_ENABLE_STAT3
-      sqlite3ExprCodeLoadIndexColumn(pParse, pIdx, iTabCur, 0, regSample);
-#else
       for(i=0; i<nCol; i++){
         sqlite3ExprCodeLoadIndexColumn(pParse, pIdx, iTabCur, i, regCol+i);
       }
       sqlite3VdbeAddOp3(v, OP_MakeRecord, regCol, nCol, regSample);
-#endif
       sqlite3VdbeAddOp3(v, OP_MakeRecord, regTabname, 6, regTemp);
       sqlite3VdbeAddOp2(v, OP_NewRowid, iStatCur+1, regNewRowid);
       sqlite3VdbeAddOp3(v, OP_Insert, iStatCur+1, regTemp, regNewRowid);
       sqlite3VdbeAddOp2(v, OP_Goto, 1, addrNext); /* P1==1 for end-of-loop */
       sqlite3VdbeJumpHere(v, addrIsNull);
     }
-#endif /* SQLITE_ENABLE_STAT3_OR_STAT4 */
+#endif /* SQLITE_ENABLE_STAT4 */
 
     /* End of analysis */
     sqlite3VdbeJumpHere(v, addrRewind);
@@ -105456,7 +107740,7 @@ static void decodeIntArray(
   int i;
   tRowcnt v;
 
-#ifdef SQLITE_ENABLE_STAT3_OR_STAT4
+#ifdef SQLITE_ENABLE_STAT4
   if( z==0 ) z = "";
 #else
   assert( z!=0 );
@@ -105467,7 +107751,7 @@ static void decodeIntArray(
       v = v*10 + c - '0';
       z++;
     }
-#ifdef SQLITE_ENABLE_STAT3_OR_STAT4
+#ifdef SQLITE_ENABLE_STAT4
     if( aOut ) aOut[i] = v;
     if( aLog ) aLog[i] = sqlite3LogEst(v);
 #else
@@ -105478,7 +107762,7 @@ static void decodeIntArray(
 #endif
     if( *z==' ' ) z++;
   }
-#ifndef SQLITE_ENABLE_STAT3_OR_STAT4
+#ifndef SQLITE_ENABLE_STAT4
   assert( pIndex!=0 ); {
 #else
   if( pIndex ){
@@ -105489,7 +107773,9 @@ static void decodeIntArray(
       if( sqlite3_strglob("unordered*", z)==0 ){
         pIndex->bUnordered = 1;
       }else if( sqlite3_strglob("sz=[0-9]*", z)==0 ){
-        pIndex->szIdxRow = sqlite3LogEst(sqlite3Atoi(z+3));
+        int sz = sqlite3Atoi(z+3);
+        if( sz<2 ) sz = 2;
+        pIndex->szIdxRow = sqlite3LogEst(sz);
       }else if( sqlite3_strglob("noskipscan*", z)==0 ){
         pIndex->noSkipScan = 1;
       }
@@ -105543,7 +107829,7 @@ static int analysisLoader(void *pData, int argc, char **argv, char **NotUsed){
   if( pIndex ){
     tRowcnt *aiRowEst = 0;
     int nCol = pIndex->nKeyCol+1;
-#ifdef SQLITE_ENABLE_STAT3_OR_STAT4
+#ifdef SQLITE_ENABLE_STAT4
     /* Index.aiRowEst may already be set here if there are duplicate 
     ** sqlite_stat1 entries for this index. In that case just clobber
     ** the old data with the new instead of allocating a new array.  */
@@ -105579,7 +107865,7 @@ static int analysisLoader(void *pData, int argc, char **argv, char **NotUsed){
 ** and its contents.
 */
 SQLITE_PRIVATE void sqlite3DeleteIndexSamples(sqlite3 *db, Index *pIdx){
-#ifdef SQLITE_ENABLE_STAT3_OR_STAT4
+#ifdef SQLITE_ENABLE_STAT4
   if( pIdx->aSample ){
     int j;
     for(j=0; j<pIdx->nSample; j++){
@@ -105595,10 +107881,10 @@ SQLITE_PRIVATE void sqlite3DeleteIndexSamples(sqlite3 *db, Index *pIdx){
 #else
   UNUSED_PARAMETER(db);
   UNUSED_PARAMETER(pIdx);
-#endif /* SQLITE_ENABLE_STAT3_OR_STAT4 */
+#endif /* SQLITE_ENABLE_STAT4 */
 }
 
-#ifdef SQLITE_ENABLE_STAT3_OR_STAT4
+#ifdef SQLITE_ENABLE_STAT4
 /*
 ** Populate the pIdx->aAvgEq[] array based on the samples currently
 ** stored in pIdx->aSample[]. 
@@ -105676,12 +107962,11 @@ static Index *findIndexOrPrimaryKey(
 }
 
 /*
-** Load the content from either the sqlite_stat4 or sqlite_stat3 table 
+** Load the content from either the sqlite_stat4
 ** into the relevant Index.aSample[] arrays.
 **
 ** Arguments zSql1 and zSql2 must point to SQL statements that return
-** data equivalent to the following (statements are different for stat3,
-** see the caller of this function for details):
+** data equivalent to the following:
 **
 **    zSql1: SELECT idx,count(*) FROM %Q.sqlite_stat4 GROUP BY idx
 **    zSql2: SELECT idx,neq,nlt,ndlt,sample FROM %Q.sqlite_stat4
@@ -105690,7 +107975,6 @@ static Index *findIndexOrPrimaryKey(
 */
 static int loadStatTbl(
   sqlite3 *db,                  /* Database handle */
-  int bStat3,                   /* Assume single column records only */
   const char *zSql1,            /* SQL statement 1 (see above) */
   const char *zSql2,            /* SQL statement 2 (see above) */
   const char *zDb               /* Database name (e.g. "main") */
@@ -105724,17 +108008,13 @@ static int loadStatTbl(
     if( zIndex==0 ) continue;
     nSample = sqlite3_column_int(pStmt, 1);
     pIdx = findIndexOrPrimaryKey(db, zIndex, zDb);
-    assert( pIdx==0 || bStat3 || pIdx->nSample==0 );
-    /* Index.nSample is non-zero at this point if data has already been
-    ** loaded from the stat4 table. In this case ignore stat3 data.  */
-    if( pIdx==0 || pIdx->nSample ) continue;
-    if( bStat3==0 ){
-      assert( !HasRowid(pIdx->pTable) || pIdx->nColumn==pIdx->nKeyCol+1 );
-      if( !HasRowid(pIdx->pTable) && IsPrimaryKeyIndex(pIdx) ){
-        nIdxCol = pIdx->nKeyCol;
-      }else{
-        nIdxCol = pIdx->nColumn;
-      }
+    assert( pIdx==0 || pIdx->nSample==0 );
+    if( pIdx==0 ) continue;
+    assert( !HasRowid(pIdx->pTable) || pIdx->nColumn==pIdx->nKeyCol+1 );
+    if( !HasRowid(pIdx->pTable) && IsPrimaryKeyIndex(pIdx) ){
+      nIdxCol = pIdx->nKeyCol;
+    }else{
+      nIdxCol = pIdx->nColumn;
     }
     pIdx->nSampleCol = nIdxCol;
     nByte = sizeof(IndexSample) * nSample;
@@ -105776,9 +108056,8 @@ static int loadStatTbl(
     pIdx = findIndexOrPrimaryKey(db, zIndex, zDb);
     if( pIdx==0 ) continue;
     /* This next condition is true if data has already been loaded from 
-    ** the sqlite_stat4 table. In this case ignore stat3 data.  */
+    ** the sqlite_stat4 table. */
     nCol = pIdx->nSampleCol;
-    if( bStat3 && nCol>1 ) continue;
     if( pIdx!=pPrevIdx ){
       initAvgEq(pPrevIdx);
       pPrevIdx = pIdx;
@@ -105811,7 +108090,7 @@ static int loadStatTbl(
 }
 
 /*
-** Load content from the sqlite_stat4 and sqlite_stat3 tables into 
+** Load content from the sqlite_stat4 table into 
 ** the Index.aSample[] arrays of all indices.
 */
 static int loadStat4(sqlite3 *db, const char *zDb){
@@ -105819,37 +108098,28 @@ static int loadStat4(sqlite3 *db, const char *zDb){
 
   assert( db->lookaside.bDisable );
   if( sqlite3FindTable(db, "sqlite_stat4", zDb) ){
-    rc = loadStatTbl(db, 0,
+    rc = loadStatTbl(db,
       "SELECT idx,count(*) FROM %Q.sqlite_stat4 GROUP BY idx", 
       "SELECT idx,neq,nlt,ndlt,sample FROM %Q.sqlite_stat4",
       zDb
     );
   }
-
-  if( rc==SQLITE_OK && sqlite3FindTable(db, "sqlite_stat3", zDb) ){
-    rc = loadStatTbl(db, 1,
-      "SELECT idx,count(*) FROM %Q.sqlite_stat3 GROUP BY idx", 
-      "SELECT idx,neq,nlt,ndlt,sqlite_record(sample) FROM %Q.sqlite_stat3",
-      zDb
-    );
-  }
-
   return rc;
 }
-#endif /* SQLITE_ENABLE_STAT3_OR_STAT4 */
+#endif /* SQLITE_ENABLE_STAT4 */
 
 /*
-** Load the content of the sqlite_stat1 and sqlite_stat3/4 tables. The
+** Load the content of the sqlite_stat1 and sqlite_stat4 tables. The
 ** contents of sqlite_stat1 are used to populate the Index.aiRowEst[]
-** arrays. The contents of sqlite_stat3/4 are used to populate the
+** arrays. The contents of sqlite_stat4 are used to populate the
 ** Index.aSample[] arrays.
 **
 ** If the sqlite_stat1 table is not present in the database, SQLITE_ERROR
-** is returned. In this case, even if SQLITE_ENABLE_STAT3/4 was defined 
-** during compilation and the sqlite_stat3/4 table is present, no data is 
+** is returned. In this case, even if SQLITE_ENABLE_STAT4 was defined 
+** during compilation and the sqlite_stat4 table is present, no data is 
 ** read from it.
 **
-** If SQLITE_ENABLE_STAT3/4 was defined during compilation and the 
+** If SQLITE_ENABLE_STAT4 was defined during compilation and the 
 ** sqlite_stat4 table is not present in the database, SQLITE_ERROR is
 ** returned. However, in this case, data is read from the sqlite_stat1
 ** table (if it is present) before returning.
@@ -105877,7 +108147,7 @@ SQLITE_PRIVATE int sqlite3AnalysisLoad(sqlite3 *db, int iDb){
   for(i=sqliteHashFirst(&pSchema->idxHash); i; i=sqliteHashNext(i)){
     Index *pIdx = sqliteHashData(i);
     pIdx->hasStat1 = 0;
-#ifdef SQLITE_ENABLE_STAT3_OR_STAT4
+#ifdef SQLITE_ENABLE_STAT4
     sqlite3DeleteIndexSamples(db, pIdx);
     pIdx->aSample = 0;
 #endif
@@ -105905,11 +108175,11 @@ SQLITE_PRIVATE int sqlite3AnalysisLoad(sqlite3 *db, int iDb){
   }
 
   /* Load the statistics from the sqlite_stat4 table. */
-#ifdef SQLITE_ENABLE_STAT3_OR_STAT4
+#ifdef SQLITE_ENABLE_STAT4
   if( rc==SQLITE_OK ){
-    db->lookaside.bDisable++;
+    DisableLookaside;
     rc = loadStat4(db, sInfo.zDatabase);
-    db->lookaside.bDisable--;
+    EnableLookaside;
   }
   for(i=sqliteHashFirst(&pSchema->idxHash); i; i=sqliteHashNext(i)){
     Index *pIdx = sqliteHashData(i);
@@ -106230,6 +108500,7 @@ static void detachFunc(
   sqlite3 *db = sqlite3_context_db_handle(context);
   int i;
   Db *pDb = 0;
+  HashElem *pEntry;
   char zErr[128];
 
   UNUSED_PARAMETER(NotUsed);
@@ -106254,6 +108525,18 @@ static void detachFunc(
     goto detach_error;
   }
 
+  /* If any TEMP triggers reference the schema being detached, move those
+  ** triggers to reference the TEMP schema itself. */
+  assert( db->aDb[1].pSchema );
+  pEntry = sqliteHashFirst(&db->aDb[1].pSchema->trigHash);
+  while( pEntry ){
+    Trigger *pTrig = (Trigger*)sqliteHashData(pEntry);
+    if( pTrig->pTabSchema==pDb->pSchema ){
+      pTrig->pTabSchema = pTrig->pSchema;
+    }
+    pEntry = sqliteHashNext(pEntry);
+  }
+
   sqlite3BtreeClose(pDb->pBt);
   pDb->pBt = 0;
   pDb->pSchema = 0;
@@ -106319,11 +108602,8 @@ static void codeAttach(
 
   assert( v || db->mallocFailed );
   if( v ){
-    sqlite3VdbeAddOp4(v, OP_Function0, 0, regArgs+3-pFunc->nArg, regArgs+3,
-                      (char *)pFunc, P4_FUNCDEF);
-    assert( pFunc->nArg==-1 || (pFunc->nArg&0xff)==pFunc->nArg );
-    sqlite3VdbeChangeP5(v, (u8)(pFunc->nArg));
- 
+    sqlite3VdbeAddFunctionCall(pParse, 0, regArgs+3-pFunc->nArg, regArgs+3,
+                               pFunc->nArg, pFunc, 0);
     /* Code an OP_Expire. For an ATTACH statement, set P1 to true (expire this
     ** statement only). For DETACH, set it to false (expire all existing
     ** statements).
@@ -106398,7 +108678,7 @@ SQLITE_PRIVATE void sqlite3FixInit(
   pFix->pSchema = db->aDb[iDb].pSchema;
   pFix->zType = zType;
   pFix->pName = pName;
-  pFix->bVarOnly = (iDb==1);
+  pFix->bTemp = (iDb==1);
 }
 
 /*
@@ -106426,7 +108706,7 @@ SQLITE_PRIVATE int sqlite3FixSrcList(
   if( NEVER(pList==0) ) return 0;
   zDb = pFix->zDb;
   for(i=0, pItem=pList->a; i<pList->nSrc; i++, pItem++){
-    if( pFix->bVarOnly==0 ){
+    if( pFix->bTemp==0 ){
       if( pItem->zDatabase && sqlite3StrICmp(pItem->zDatabase, zDb) ){
         sqlite3ErrorMsg(pFix->pParse,
             "%s %T cannot reference objects in database %s",
@@ -106436,6 +108716,7 @@ SQLITE_PRIVATE int sqlite3FixSrcList(
       sqlite3DbFree(pFix->pParse->db, pItem->zDatabase);
       pItem->zDatabase = 0;
       pItem->pSchema = pFix->pSchema;
+      pItem->fg.fromDDL = 1;
     }
 #if !defined(SQLITE_OMIT_VIEW) || !defined(SQLITE_OMIT_TRIGGER)
     if( sqlite3FixSelect(pFix, pItem->pSelect) ) return 1;
@@ -106491,6 +108772,7 @@ SQLITE_PRIVATE int sqlite3FixExpr(
   Expr *pExpr        /* The expression to be fixed to one database */
 ){
   while( pExpr ){
+    if( !pFix->bTemp ) ExprSetProperty(pExpr, EP_FromDDL);
     if( pExpr->op==TK_VARIABLE ){
       if( pFix->pParse->db->init.busy ){
         pExpr->op = TK_NULL;
@@ -106643,7 +108925,7 @@ SQLITE_API int sqlite3_set_authorizer(
   sqlite3_mutex_enter(db->mutex);
   db->xAuth = (sqlite3_xauth)xAuth;
   db->pAuthArg = pArg;
-  sqlite3ExpirePreparedStatements(db, 0);
+  if( db->xAuth ) sqlite3ExpirePreparedStatements(db, 1);
   sqlite3_mutex_leave(db->mutex);
   return SQLITE_OK;
 }
@@ -107297,7 +109579,7 @@ SQLITE_PRIVATE void sqlite3FreeIndex(sqlite3 *db, Index *p){
   sqlite3ExprListDelete(db, p->aColExpr);
   sqlite3DbFree(db, p->zColAff);
   if( p->isResized ) sqlite3DbFree(db, (void *)p->azColl);
-#ifdef SQLITE_ENABLE_STAT3_OR_STAT4
+#ifdef SQLITE_ENABLE_STAT4
   sqlite3_free(p->aiRowEst);
 #endif
   sqlite3DbFree(db, p);
@@ -107459,10 +109741,14 @@ static void SQLITE_NOINLINE deleteTable(sqlite3 *db, Table *pTable){
 
 #ifdef SQLITE_DEBUG
   /* Record the number of outstanding lookaside allocations in schema Tables
-  ** prior to doing any free() operations.  Since schema Tables do not use
-  ** lookaside, this number should not change. */
+  ** prior to doing any free() operations. Since schema Tables do not use
+  ** lookaside, this number should not change. 
+  **
+  ** If malloc has already failed, it may be that it failed while allocating
+  ** a Table object that was going to be marked ephemeral. So do not check
+  ** that no lookaside memory is used in this case either. */
   int nLookaside = 0;
-  if( db && (pTable->tabFlags & TF_Ephemeral)==0 ){
+  if( db && !db->mallocFailed && (pTable->tabFlags & TF_Ephemeral)==0 ){
     nLookaside = sqlite3LookasideUsed(db, 0);
   }
 #endif
@@ -107666,13 +109952,41 @@ SQLITE_PRIVATE int sqlite3WritableSchema(sqlite3 *db){
 ** trigger). All names are legal except those that begin with the string
 ** "sqlite_" (in upper, lower or mixed case). This portion of the namespace
 ** is reserved for internal use.
+**
+** When parsing the sqlite_master table, this routine also checks to
+** make sure the "type", "name", and "tbl_name" columns are consistent
+** with the SQL.
 */
-SQLITE_PRIVATE int sqlite3CheckObjectName(Parse *pParse, const char *zName){
-  if( !pParse->db->init.busy && pParse->nested==0 
-          && sqlite3WritableSchema(pParse->db)==0
-          && 0==sqlite3StrNICmp(zName, "sqlite_", 7) ){
-    sqlite3ErrorMsg(pParse, "object name reserved for internal use: %s", zName);
-    return SQLITE_ERROR;
+SQLITE_PRIVATE int sqlite3CheckObjectName(
+  Parse *pParse,            /* Parsing context */
+  const char *zName,        /* Name of the object to check */
+  const char *zType,        /* Type of this object */
+  const char *zTblName      /* Parent table name for triggers and indexes */
+){
+  sqlite3 *db = pParse->db;
+  if( sqlite3WritableSchema(db) || db->init.imposterTable ){
+    /* Skip these error checks for writable_schema=ON */
+    return SQLITE_OK;
+  }
+  if( db->init.busy ){
+    if( sqlite3_stricmp(zType, db->init.azInit[0])
+     || sqlite3_stricmp(zName, db->init.azInit[1])
+     || sqlite3_stricmp(zTblName, db->init.azInit[2])
+    ){
+      if( sqlite3Config.bExtraSchemaChecks ){
+        sqlite3ErrorMsg(pParse, ""); /* corruptSchema() will supply the error */
+        return SQLITE_ERROR;
+      }
+    }
+  }else{
+    if( (pParse->nested==0 && 0==sqlite3StrNICmp(zName, "sqlite_", 7))
+     || (sqlite3ReadOnlyShadowTables(db) && sqlite3ShadowTableName(db, zName))
+    ){
+      sqlite3ErrorMsg(pParse, "object name reserved for internal use: %s",
+                      zName);
+      return SQLITE_ERROR;
+    }
+
   }
   return SQLITE_OK;
 }
@@ -107687,10 +110001,12 @@ SQLITE_PRIVATE Index *sqlite3PrimaryKeyIndex(Table *pTab){
 }
 
 /*
-** Return the column of index pIdx that corresponds to table
-** column iCol.  Return -1 if not found.
+** Convert an table column number into a index column number.  That is,
+** for the column iCol in the table (as defined by the CREATE TABLE statement)
+** find the (first) offset of that column in index pIdx.  Or return -1
+** if column iCol is not used in index pIdx.
 */
-SQLITE_PRIVATE i16 sqlite3ColumnOfIndex(Index *pIdx, i16 iCol){
+SQLITE_PRIVATE i16 sqlite3TableColumnToIndex(Index *pIdx, i16 iCol){
   int i;
   for(i=0; i<pIdx->nColumn; i++){
     if( iCol==pIdx->aiColumn[i] ) return i;
@@ -107698,6 +110014,84 @@ SQLITE_PRIVATE i16 sqlite3ColumnOfIndex(Index *pIdx, i16 iCol){
   return -1;
 }
 
+#ifndef SQLITE_OMIT_GENERATED_COLUMNS
+/* Convert a storage column number into a table column number.
+**
+** The storage column number (0,1,2,....) is the index of the value
+** as it appears in the record on disk.  The true column number
+** is the index (0,1,2,...) of the column in the CREATE TABLE statement.
+**
+** The storage column number is less than the table column number if
+** and only there are VIRTUAL columns to the left.
+**
+** If SQLITE_OMIT_GENERATED_COLUMNS, this routine is a no-op macro.
+*/
+SQLITE_PRIVATE i16 sqlite3StorageColumnToTable(Table *pTab, i16 iCol){
+  if( pTab->tabFlags & TF_HasVirtual ){
+    int i;
+    for(i=0; i<=iCol; i++){
+      if( pTab->aCol[i].colFlags & COLFLAG_VIRTUAL ) iCol++;
+    }
+  }
+  return iCol;
+}
+#endif
+
+#ifndef SQLITE_OMIT_GENERATED_COLUMNS
+/* Convert a table column number into a storage column number.
+**
+** The storage column number (0,1,2,....) is the index of the value
+** as it appears in the record on disk.  Or, if the input column is
+** the N-th virtual column (zero-based) then the storage number is
+** the number of non-virtual columns in the table plus N.  
+**
+** The true column number is the index (0,1,2,...) of the column in
+** the CREATE TABLE statement.
+**
+** If the input column is a VIRTUAL column, then it should not appear
+** in storage.  But the value sometimes is cached in registers that
+** follow the range of registers used to construct storage.  This
+** avoids computing the same VIRTUAL column multiple times, and provides
+** values for use by OP_Param opcodes in triggers.  Hence, if the
+** input column is a VIRTUAL table, put it after all the other columns.
+**
+** In the following, N means "normal column", S means STORED, and
+** V means VIRTUAL.  Suppose the CREATE TABLE has columns like this:
+**
+**        CREATE TABLE ex(N,S,V,N,S,V,N,S,V);
+**                     -- 0 1 2 3 4 5 6 7 8
+**
+** Then the mapping from this function is as follows:
+**
+**    INPUTS:     0 1 2 3 4 5 6 7 8
+**    OUTPUTS:    0 1 6 2 3 7 4 5 8
+**
+** So, in other words, this routine shifts all the virtual columns to
+** the end.
+**
+** If SQLITE_OMIT_GENERATED_COLUMNS then there are no virtual columns and
+** this routine is a no-op macro.  If the pTab does not have any virtual
+** columns, then this routine is no-op that always return iCol.  If iCol
+** is negative (indicating the ROWID column) then this routine return iCol.
+*/
+SQLITE_PRIVATE i16 sqlite3TableColumnToStorage(Table *pTab, i16 iCol){
+  int i;
+  i16 n;
+  assert( iCol<pTab->nCol );
+  if( (pTab->tabFlags & TF_HasVirtual)==0 || iCol<0 ) return iCol;
+  for(i=0, n=0; i<iCol; i++){
+    if( (pTab->aCol[i].colFlags & COLFLAG_VIRTUAL)==0 ) n++;
+  }
+  if( pTab->aCol[i].colFlags & COLFLAG_VIRTUAL ){
+    /* iCol is a virtual column itself */
+    return pTab->nNVCol + i - n;
+  }else{
+    /* iCol is a normal or stored column */
+    return n;
+  }
+}
+#endif
+
 /*
 ** Begin constructing a new table representation in memory.  This is
 ** the first of several action routines that get called in response
@@ -107753,7 +110147,7 @@ SQLITE_PRIVATE void sqlite3StartTable(
   }
   pParse->sNameToken = *pName;
   if( zName==0 ) return;
-  if( SQLITE_OK!=sqlite3CheckObjectName(pParse, zName) ){
+  if( sqlite3CheckObjectName(pParse, zName, isView?"view":"table", zName) ){
     goto begin_table_error;
   }
   if( db->init.iDb==1 ) isTemp = 1;
@@ -107988,6 +110382,7 @@ SQLITE_PRIVATE void sqlite3AddColumn(Parse *pParse, Token *pName, Token *pType){
     pCol->colFlags |= COLFLAG_HASTYPE;
   }
   p->nCol++;
+  p->nNVCol++;
   pParse->constraintName.n = 0;
 }
 
@@ -108132,10 +110527,17 @@ SQLITE_PRIVATE void sqlite3AddDefaultValue(
   sqlite3 *db = pParse->db;
   p = pParse->pNewTable;
   if( p!=0 ){
+    int isInit = db->init.busy && db->init.iDb!=1;
     pCol = &(p->aCol[p->nCol-1]);
-    if( !sqlite3ExprIsConstantOrFunction(pExpr, db->init.busy) ){
+    if( !sqlite3ExprIsConstantOrFunction(pExpr, isInit) ){
       sqlite3ErrorMsg(pParse, "default value of column [%s] is not constant",
           pCol->zName);
+#ifndef SQLITE_OMIT_GENERATED_COLUMNS
+    }else if( pCol->colFlags & COLFLAG_GENERATED ){
+      testcase( pCol->colFlags & COLFLAG_VIRTUAL );
+      testcase( pCol->colFlags & COLFLAG_STORED );
+      sqlite3ErrorMsg(pParse, "cannot use DEFAULT on a generated column");
+#endif
     }else{
       /* A copy of pExpr is used instead of the original, as pExpr contains
       ** tokens that point to volatile memory.
@@ -108170,7 +110572,7 @@ SQLITE_PRIVATE void sqlite3AddDefaultValue(
 ** accept it.  This routine does the necessary conversion.  It converts
 ** the expression given in its argument from a TK_STRING into a TK_ID
 ** if the expression is just a TK_STRING with an optional COLLATE clause.
-** If the epxression is anything other than TK_STRING, the expression is
+** If the expression is anything other than TK_STRING, the expression is
 ** unchanged.
 */
 static void sqlite3StringToId(Expr *p){
@@ -108181,6 +110583,21 @@ static void sqlite3StringToId(Expr *p){
   }
 }
 
+/*
+** Tag the given column as being part of the PRIMARY KEY
+*/
+static void makeColumnPartOfPrimaryKey(Parse *pParse, Column *pCol){
+  pCol->colFlags |= COLFLAG_PRIMKEY;
+#ifndef SQLITE_OMIT_GENERATED_COLUMNS
+  if( pCol->colFlags & COLFLAG_GENERATED ){
+    testcase( pCol->colFlags & COLFLAG_VIRTUAL );
+    testcase( pCol->colFlags & COLFLAG_STORED );
+    sqlite3ErrorMsg(pParse,
+      "generated columns cannot be part of the PRIMARY KEY");
+  }
+#endif          
+}
+
 /*
 ** Designate the PRIMARY KEY for the table.  pList is a list of names 
 ** of columns that form the primary key.  If pList is NULL, then the
@@ -108220,7 +110637,7 @@ SQLITE_PRIVATE void sqlite3AddPrimaryKey(
   if( pList==0 ){
     iCol = pTab->nCol - 1;
     pCol = &pTab->aCol[iCol];
-    pCol->colFlags |= COLFLAG_PRIMKEY;
+    makeColumnPartOfPrimaryKey(pParse, pCol);
     nTerm = 1;
   }else{
     nTerm = pList->nExpr;
@@ -108233,7 +110650,7 @@ SQLITE_PRIVATE void sqlite3AddPrimaryKey(
         for(iCol=0; iCol<pTab->nCol; iCol++){
           if( sqlite3StrICmp(zCName, pTab->aCol[iCol].zName)==0 ){
             pCol = &pTab->aCol[iCol];
-            pCol->colFlags |= COLFLAG_PRIMKEY;
+            makeColumnPartOfPrimaryKey(pParse, pCol);
             break;
           }
         }
@@ -108253,7 +110670,8 @@ SQLITE_PRIVATE void sqlite3AddPrimaryKey(
     pTab->keyConf = (u8)onError;
     assert( autoInc==0 || autoInc==1 );
     pTab->tabFlags |= autoInc*TF_Autoincrement;
-    if( pList ) pParse->iPkSortOrder = pList->a[0].sortOrder;
+    if( pList ) pParse->iPkSortOrder = pList->a[0].sortFlags;
+    (void)sqlite3HasExplicitNulls(pParse, pList);
   }else if( autoInc ){
 #ifndef SQLITE_OMIT_AUTOINCREMENT
     sqlite3ErrorMsg(pParse, "AUTOINCREMENT is only allowed on an "
@@ -108330,41 +110748,58 @@ SQLITE_PRIVATE void sqlite3AddCollateType(Parse *pParse, Token *pToken){
   }
 }
 
-/*
-** This function returns the collation sequence for database native text
-** encoding identified by the string zName, length nName.
-**
-** If the requested collation sequence is not available, or not available
-** in the database native encoding, the collation factory is invoked to
-** request it. If the collation factory does not supply such a sequence,
-** and the sequence is available in another text encoding, then that is
-** returned instead.
-**
-** If no versions of the requested collations sequence are available, or
-** another error occurs, NULL is returned and an error message written into
-** pParse.
-**
-** This routine is a wrapper around sqlite3FindCollSeq().  This routine
-** invokes the collation factory if the named collation cannot be found
-** and generates an error message.
-**
-** See also: sqlite3FindCollSeq(), sqlite3GetCollSeq()
+/* Change the most recently parsed column to be a GENERATED ALWAYS AS
+** column.
 */
-SQLITE_PRIVATE CollSeq *sqlite3LocateCollSeq(Parse *pParse, const char *zName){
-  sqlite3 *db = pParse->db;
-  u8 enc = ENC(db);
-  u8 initbusy = db->init.busy;
-  CollSeq *pColl;
-
-  pColl = sqlite3FindCollSeq(db, enc, zName, initbusy);
-  if( !initbusy && (!pColl || !pColl->xCmp) ){
-    pColl = sqlite3GetCollSeq(pParse, enc, pColl, zName);
+SQLITE_PRIVATE void sqlite3AddGenerated(Parse *pParse, Expr *pExpr, Token *pType){
+#ifndef SQLITE_OMIT_GENERATED_COLUMNS
+  u8 eType = COLFLAG_VIRTUAL;
+  Table *pTab = pParse->pNewTable;
+  Column *pCol;
+  if( pTab==0 ){
+    /* generated column in an CREATE TABLE IF NOT EXISTS that already exists */
+    goto generated_done;
+  }
+  pCol = &(pTab->aCol[pTab->nCol-1]);
+  if( IN_DECLARE_VTAB ){
+    sqlite3ErrorMsg(pParse, "virtual tables cannot use computed columns");
+    goto generated_done;
+  }
+  if( pCol->pDflt ) goto generated_error;
+  if( pType ){
+    if( pType->n==7 && sqlite3StrNICmp("virtual",pType->z,7)==0 ){
+      /* no-op */
+    }else if( pType->n==6 && sqlite3StrNICmp("stored",pType->z,6)==0 ){
+      eType = COLFLAG_STORED;
+    }else{
+      goto generated_error;
+    }
+  }
+  if( eType==COLFLAG_VIRTUAL ) pTab->nNVCol--;
+  pCol->colFlags |= eType;
+  assert( TF_HasVirtual==COLFLAG_VIRTUAL );
+  assert( TF_HasStored==COLFLAG_STORED );
+  pTab->tabFlags |= eType;
+  if( pCol->colFlags & COLFLAG_PRIMKEY ){
+    makeColumnPartOfPrimaryKey(pParse, pCol); /* For the error message */
   }
+  pCol->pDflt = pExpr;
+  pExpr = 0;
+  goto generated_done;
 
-  return pColl;
+generated_error:
+  sqlite3ErrorMsg(pParse, "error in generated column \"%s\"",
+                  pCol->zName);
+generated_done:
+  sqlite3ExprDelete(pParse->db, pExpr);
+#else
+  /* Throw and error for the GENERATED ALWAYS AS clause if the
+  ** SQLITE_OMIT_GENERATED_COLUMNS compile-time option is used. */
+  sqlite3ErrorMsg(pParse, "generated columns not supported");
+  sqlite3ExprDelete(pParse->db, pExpr);
+#endif
 }
 
-
 /*
 ** Generate code that will increment the schema cookie.
 **
@@ -108567,10 +111002,51 @@ static void estimateIndexWidth(Index *pIdx){
   pIdx->szIdxRow = sqlite3LogEst(wIndex*4);
 }
 
-/* Return true if value x is found any of the first nCol entries of aiCol[]
+/* Return true if column number x is any of the first nCol entries of aiCol[].
+** This is used to determine if the column number x appears in any of the
+** first nCol entries of an index.
 */
 static int hasColumn(const i16 *aiCol, int nCol, int x){
-  while( nCol-- > 0 ) if( x==*(aiCol++) ) return 1;
+  while( nCol-- > 0 ){
+    assert( aiCol[0]>=0 );
+    if( x==*(aiCol++) ){
+      return 1;
+    }
+  }
+  return 0;
+}
+
+/*
+** Return true if any of the first nKey entries of index pIdx exactly
+** match the iCol-th entry of pPk.  pPk is always a WITHOUT ROWID
+** PRIMARY KEY index.  pIdx is an index on the same table.  pIdx may
+** or may not be the same index as pPk.
+**
+** The first nKey entries of pIdx are guaranteed to be ordinary columns,
+** not a rowid or expression.
+**
+** This routine differs from hasColumn() in that both the column and the
+** collating sequence must match for this routine, but for hasColumn() only
+** the column name must match.
+*/
+static int isDupColumn(Index *pIdx, int nKey, Index *pPk, int iCol){
+  int i, j;
+  assert( nKey<=pIdx->nColumn );
+  assert( iCol<MAX(pPk->nColumn,pPk->nKeyCol) );
+  assert( pPk->idxType==SQLITE_IDXTYPE_PRIMARYKEY );
+  assert( pPk->pTable->tabFlags & TF_WithoutRowid );
+  assert( pPk->pTable==pIdx->pTable );
+  testcase( pPk==pIdx );
+  j = pPk->aiColumn[iCol];
+  assert( j!=XN_ROWID && j!=XN_EXPR );
+  for(i=0; i<nKey; i++){
+    assert( pIdx->aiColumn[i]>=0 || j>=0 );
+    if( pIdx->aiColumn[i]==j 
+     && sqlite3StrICmp(pIdx->azColl[i], pPk->azColl[iCol])==0
+    ){
+      return 1;
+    }
+  }
   return 0;
 }
 
@@ -108581,15 +111057,24 @@ static int hasColumn(const i16 *aiCol, int nCol, int x){
 ** high-order bit of colNotIdxed is always 1.  All unindexed columns
 ** of the table have a 1.
 **
+** 2019-10-24:  For the purpose of this computation, virtual columns are
+** not considered to be covered by the index, even if they are in the
+** index, because we do not trust the logic in whereIndexExprTrans() to be
+** able to find all instances of a reference to the indexed table column
+** and convert them into references to the index.  Hence we always want
+** the actual table at hand in order to recompute the virtual column, if
+** necessary.
+**
 ** The colNotIdxed mask is AND-ed with the SrcList.a[].colUsed mask
 ** to determine if the index is covering index.
 */
 static void recomputeColumnsNotIndexed(Index *pIdx){
   Bitmask m = 0;
   int j;
+  Table *pTab = pIdx->pTable;
   for(j=pIdx->nColumn-1; j>=0; j--){
     int x = pIdx->aiColumn[j];
-    if( x>=0 ){
+    if( x>=0 && (pTab->aCol[x].colFlags & COLFLAG_VIRTUAL)==0 ){
       testcase( x==BMS-1 );
       testcase( x==BMS-2 );
       if( x<BMS-1 ) m |= MASKBIT(x);
@@ -108627,6 +111112,7 @@ static void convertToWithoutRowidTable(Parse *pParse, Table *pTab){
   Index *pIdx;
   Index *pPk;
   int nPk;
+  int nExtra;
   int i, j;
   sqlite3 *db = pParse->db;
   Vdbe *v = pParse->pVdbe;
@@ -108639,6 +111125,7 @@ static void convertToWithoutRowidTable(Parse *pParse, Table *pTab){
         pTab->aCol[i].notNull = OE_Abort;
       }
     }
+    pTab->tabFlags |= TF_HasNotNull;
   }
 
   /* Convert the P3 operand of the OP_CreateBtree opcode from BTREE_INTKEY
@@ -108659,13 +111146,17 @@ static void convertToWithoutRowidTable(Parse *pParse, Table *pTab){
     pList = sqlite3ExprListAppend(pParse, 0, 
                   sqlite3ExprAlloc(db, TK_ID, &ipkToken, 0));
     if( pList==0 ) return;
-    pList->a[0].sortOrder = pParse->iPkSortOrder;
+    if( IN_RENAME_OBJECT ){
+      sqlite3RenameTokenRemap(pParse, pList->a[0].pExpr, &pTab->iPKey);
+    }
+    pList->a[0].sortFlags = pParse->iPkSortOrder;
     assert( pParse->pNewTable==pTab );
+    pTab->iPKey = -1;
     sqlite3CreateIndex(pParse, 0, 0, 0, pList, pTab->keyConf, 0, 0, 0, 0,
                        SQLITE_IDXTYPE_PRIMARYKEY);
     if( db->mallocFailed || pParse->nErr ) return;
     pPk = sqlite3PrimaryKeyIndex(pTab);
-    pTab->iPKey = -1;
+    assert( pPk->nKeyCol==1 );
   }else{
     pPk = sqlite3PrimaryKeyIndex(pTab);
     assert( pPk!=0 );
@@ -108676,9 +111167,12 @@ static void convertToWithoutRowidTable(Parse *pParse, Table *pTab){
     ** code assumes the PRIMARY KEY contains no repeated columns.
     */
     for(i=j=1; i<pPk->nKeyCol; i++){
-      if( hasColumn(pPk->aiColumn, j, pPk->aiColumn[i]) ){
+      if( isDupColumn(pPk, j, pPk, i) ){
         pPk->nColumn--;
       }else{
+        testcase( hasColumn(pPk->aiColumn, j, pPk->aiColumn[i]) );
+        pPk->azColl[j] = pPk->azColl[i];
+        pPk->aSortOrder[j] = pPk->aSortOrder[i];
         pPk->aiColumn[j++] = pPk->aiColumn[i];
       }
     }
@@ -108687,7 +111181,7 @@ static void convertToWithoutRowidTable(Parse *pParse, Table *pTab){
   assert( pPk!=0 );
   pPk->isCovering = 1;
   if( !db->init.imposterTable ) pPk->uniqNotNull = 1;
-  nPk = pPk->nKeyCol;
+  nPk = pPk->nColumn = pPk->nKeyCol;
 
   /* Bypass the creation of the PRIMARY KEY btree and the sqlite_master
   ** table entry. This is only required if currently generating VDBE
@@ -108708,7 +111202,10 @@ static void convertToWithoutRowidTable(Parse *pParse, Table *pTab){
     int n;
     if( IsPrimaryKeyIndex(pIdx) ) continue;
     for(i=n=0; i<nPk; i++){
-      if( !hasColumn(pIdx->aiColumn, pIdx->nKeyCol, pPk->aiColumn[i]) ) n++;
+      if( !isDupColumn(pIdx, pIdx->nKeyCol, pPk, i) ){
+        testcase( hasColumn(pIdx->aiColumn, pIdx->nKeyCol, pPk->aiColumn[i]) );
+        n++;
+      }
     }
     if( n==0 ){
       /* This index is a superset of the primary key */
@@ -108717,9 +111214,14 @@ static void convertToWithoutRowidTable(Parse *pParse, Table *pTab){
     }
     if( resizeIndexObject(db, pIdx, pIdx->nKeyCol+n) ) return;
     for(i=0, j=pIdx->nKeyCol; i<nPk; i++){
-      if( !hasColumn(pIdx->aiColumn, pIdx->nKeyCol, pPk->aiColumn[i]) ){
+      if( !isDupColumn(pIdx, pIdx->nKeyCol, pPk, i) ){
+        testcase( hasColumn(pIdx->aiColumn, pIdx->nKeyCol, pPk->aiColumn[i]) );
         pIdx->aiColumn[j] = pPk->aiColumn[i];
         pIdx->azColl[j] = pPk->azColl[i];
+        if( pPk->aSortOrder[i] ){
+          /* See ticket https://www.sqlite.org/src/info/bba7b69f9849b5bf */
+          pIdx->bAscKeyBug = 1;
+        }
         j++;
       }
     }
@@ -108729,21 +111231,24 @@ static void convertToWithoutRowidTable(Parse *pParse, Table *pTab){
 
   /* Add all table columns to the PRIMARY KEY index
   */
-  if( nPk<pTab->nCol ){
-    if( resizeIndexObject(db, pPk, pTab->nCol) ) return;
-    for(i=0, j=nPk; i<pTab->nCol; i++){
-      if( !hasColumn(pPk->aiColumn, j, i) ){
-        assert( j<pPk->nColumn );
-        pPk->aiColumn[j] = i;
-        pPk->azColl[j] = sqlite3StrBINARY;
-        j++;
-      }
+  nExtra = 0;
+  for(i=0; i<pTab->nCol; i++){
+    if( !hasColumn(pPk->aiColumn, nPk, i)
+     && (pTab->aCol[i].colFlags & COLFLAG_VIRTUAL)==0 ) nExtra++;
+  }
+  if( resizeIndexObject(db, pPk, nPk+nExtra) ) return;
+  for(i=0, j=nPk; i<pTab->nCol; i++){
+    if( !hasColumn(pPk->aiColumn, j, i)
+     && (pTab->aCol[i].colFlags & COLFLAG_VIRTUAL)==0
+    ){
+      assert( j<pPk->nColumn );
+      pPk->aiColumn[j] = i;
+      pPk->azColl[j] = sqlite3StrBINARY;
+      j++;
     }
-    assert( pPk->nColumn==j );
-    assert( pTab->nCol==j );
-  }else{
-    pPk->nColumn = pTab->nCol;
   }
+  assert( pPk->nColumn==j );
+  assert( pTab->nNVCol<=j );
   recomputeColumnsNotIndexed(pPk);
 }
 
@@ -108755,7 +111260,7 @@ static void convertToWithoutRowidTable(Parse *pParse, Table *pTab){
 ** zName is temporarily modified while this routine is running, but is
 ** restored to its original value prior to this routine returning.
 */
-static int isShadowTableName(sqlite3 *db, char *zName){
+SQLITE_PRIVATE int sqlite3ShadowTableName(sqlite3 *db, const char *zName){
   char *zTail;                  /* Pointer to the last "_" in zName */
   Table *pTab;                  /* Table that zName is a shadow of */
   Module *pMod;                 /* Module for the virtual table */
@@ -108773,8 +111278,6 @@ static int isShadowTableName(sqlite3 *db, char *zName){
   if( pMod->pModule->xShadowName==0 ) return 0;
   return pMod->pModule->xShadowName(zTail+1);
 }
-#else
-# define isShadowTableName(x,y) 0
 #endif /* ifndef SQLITE_OMIT_VIRTUALTABLE */
 
 /*
@@ -108816,7 +111319,7 @@ SQLITE_PRIVATE void sqlite3EndTable(
   p = pParse->pNewTable;
   if( p==0 ) return;
 
-  if( pSelect==0 && isShadowTableName(db, p->zName) ){
+  if( pSelect==0 && sqlite3ShadowTableName(db, p->zName) ){
     p->tabFlags |= TF_Shadow;
   }
 
@@ -108852,12 +111355,11 @@ SQLITE_PRIVATE void sqlite3EndTable(
     }
     if( (p->tabFlags & TF_HasPrimaryKey)==0 ){
       sqlite3ErrorMsg(pParse, "PRIMARY KEY missing on table %s", p->zName);
-    }else{
-      p->tabFlags |= TF_WithoutRowid | TF_NoVisibleRowid;
-      convertToWithoutRowidTable(pParse, p);
+      return;
     }
+    p->tabFlags |= TF_WithoutRowid | TF_NoVisibleRowid;
+    convertToWithoutRowidTable(pParse, p);
   }
-
   iDb = sqlite3SchemaToIndex(db, p->pSchema);
 
 #ifndef SQLITE_OMIT_CHECK
@@ -108865,8 +111367,45 @@ SQLITE_PRIVATE void sqlite3EndTable(
   */
   if( p->pCheck ){
     sqlite3ResolveSelfReference(pParse, p, NC_IsCheck, 0, p->pCheck);
+    if( pParse->nErr ){
+      /* If errors are seen, delete the CHECK constraints now, else they might
+      ** actually be used if PRAGMA writable_schema=ON is set. */
+      sqlite3ExprListDelete(db, p->pCheck);
+      p->pCheck = 0;
+    }
   }
 #endif /* !defined(SQLITE_OMIT_CHECK) */
+#ifndef SQLITE_OMIT_GENERATED_COLUMNS
+  if( p->tabFlags & TF_HasGenerated ){
+    int ii, nNG = 0;
+    testcase( p->tabFlags & TF_HasVirtual );
+    testcase( p->tabFlags & TF_HasStored );
+    for(ii=0; ii<p->nCol; ii++){
+      u32 colFlags = p->aCol[ii].colFlags;
+      if( (colFlags & COLFLAG_GENERATED)!=0 ){
+        Expr *pX = p->aCol[ii].pDflt;
+        testcase( colFlags & COLFLAG_VIRTUAL );
+        testcase( colFlags & COLFLAG_STORED );
+        if( sqlite3ResolveSelfReference(pParse, p, NC_GenCol, pX, 0) ){
+          /* If there are errors in resolving the expression, change the
+          ** expression to a NULL.  This prevents code generators that operate
+          ** on the expression from inserting extra parts into the expression
+          ** tree that have been allocated from lookaside memory, which is
+          ** illegal in a schema and will lead to errors or heap corruption
+          ** when the database connection closes. */
+          sqlite3ExprDelete(db, pX);
+          p->aCol[ii].pDflt = sqlite3ExprAlloc(db, TK_NULL, 0, 0);
+        }
+      }else{
+        nNG++;
+      }
+    }
+    if( nNG==0 ){
+      sqlite3ErrorMsg(pParse, "must have at least one non-generated column");
+      return;
+    }
+  }
+#endif
 
   /* Estimate the average row size for the table and for all implied indices */
   estimateTableWidth(p);
@@ -108940,10 +111479,10 @@ SQLITE_PRIVATE void sqlite3EndTable(
       addrTop = sqlite3VdbeCurrentAddr(v) + 1;
       sqlite3VdbeAddOp3(v, OP_InitCoroutine, regYield, 0, addrTop);
       if( pParse->nErr ) return;
-      pSelTab = sqlite3ResultSetOfSelect(pParse, pSelect);
+      pSelTab = sqlite3ResultSetOfSelect(pParse, pSelect, SQLITE_AFF_BLOB);
       if( pSelTab==0 ) return;
       assert( p->aCol==0 );
-      p->nCol = pSelTab->nCol;
+      p->nCol = p->nNVCol = pSelTab->nCol;
       p->aCol = pSelTab->aCol;
       pSelTab->nCol = 0;
       pSelTab->aCol = 0;
@@ -109016,7 +111555,6 @@ SQLITE_PRIVATE void sqlite3EndTable(
            sqlite3MPrintf(db, "tbl_name='%q' AND type!='trigger'", p->zName));
   }
 
-
   /* Add the table to the in-memory representation of the database.
   */
   if( db->init.busy ){
@@ -109087,6 +111625,7 @@ SQLITE_PRIVATE void sqlite3CreateView(
   ** allocated rather than point to the input string - which means that
   ** they will persist after the current sqlite3_exec() call returns.
   */
+  pSelect->selFlags |= SF_View;
   if( IN_RENAME_OBJECT ){
     p->pSelect = pSelect;
     pSelect = 0;
@@ -109200,17 +111739,20 @@ SQLITE_PRIVATE int sqlite3ViewGetColumnNames(Parse *pParse, Table *pTable){
     n = pParse->nTab;
     sqlite3SrcListAssignCursors(pParse, pSel->pSrc);
     pTable->nCol = -1;
-    db->lookaside.bDisable++;
+    DisableLookaside;
 #ifndef SQLITE_OMIT_AUTHORIZATION
     xAuth = db->xAuth;
     db->xAuth = 0;
-    pSelTab = sqlite3ResultSetOfSelect(pParse, pSel);
+    pSelTab = sqlite3ResultSetOfSelect(pParse, pSel, SQLITE_AFF_NONE);
     db->xAuth = xAuth;
 #else
-    pSelTab = sqlite3ResultSetOfSelect(pParse, pSel);
+    pSelTab = sqlite3ResultSetOfSelect(pParse, pSel, SQLITE_AFF_NONE);
 #endif
     pParse->nTab = n;
-    if( pTable->pCheck ){
+    if( pSelTab==0 ){
+      pTable->nCol = 0;
+      nErr++;
+    }else if( pTable->pCheck ){
       /* CREATE VIEW name(arglist) AS ...
       ** The names of the columns in the table are taken from
       ** arglist which is stored in pTable->pCheck.  The pCheck field
@@ -109223,9 +111765,10 @@ SQLITE_PRIVATE int sqlite3ViewGetColumnNames(Parse *pParse, Table *pTable){
        && pParse->nErr==0
        && pTable->nCol==pSel->pEList->nExpr
       ){
-        sqlite3SelectAddColumnTypeAndCollation(pParse, pTable, pSel);
+        sqlite3SelectAddColumnTypeAndCollation(pParse, pTable, pSel,
+                                               SQLITE_AFF_NONE);
       }
-    }else if( pSelTab ){
+    }else{
       /* CREATE VIEW name AS...  without an argument list.  Construct
       ** the column names from the SELECT statement that defines the view.
       */
@@ -109235,13 +111778,11 @@ SQLITE_PRIVATE int sqlite3ViewGetColumnNames(Parse *pParse, Table *pTable){
       pSelTab->nCol = 0;
       pSelTab->aCol = 0;
       assert( sqlite3SchemaMutexHeld(db, 0, pTable->pSchema) );
-    }else{
-      pTable->nCol = 0;
-      nErr++;
     }
+    pTable->nNVCol = pTable->nCol;
     sqlite3DeleteTable(db, pSelTab);
     sqlite3SelectDelete(db, pSel);
-    db->lookaside.bDisable--;
+    EnableLookaside;
 #ifndef SQLITE_OMIT_ALTERTABLE
     pParse->eParseMode = eParseMode;
 #endif
@@ -109498,6 +112039,37 @@ SQLITE_PRIVATE void sqlite3CodeDropTable(Parse *pParse, Table *pTab, int iDb, in
   sqliteViewResetAll(db, iDb);
 }
 
+/*
+** Return TRUE if shadow tables should be read-only in the current
+** context.
+*/
+SQLITE_PRIVATE int sqlite3ReadOnlyShadowTables(sqlite3 *db){
+#ifndef SQLITE_OMIT_VIRTUALTABLE
+  if( (db->flags & SQLITE_Defensive)!=0
+   && db->pVtabCtx==0
+   && db->nVdbeExec==0
+  ){
+    return 1;
+  }
+#endif
+  return 0;
+}
+
+/*
+** Return true if it is not allowed to drop the given table
+*/
+static int tableMayNotBeDropped(sqlite3 *db, Table *pTab){
+  if( sqlite3StrNICmp(pTab->zName, "sqlite_", 7)==0 ){
+    if( sqlite3StrNICmp(pTab->zName+7, "stat", 4)==0 ) return 0;
+    if( sqlite3StrNICmp(pTab->zName+7, "parameters", 10)==0 ) return 0;
+    return 1;
+  }
+  if( (pTab->tabFlags & TF_Shadow)!=0 && sqlite3ReadOnlyShadowTables(db) ){
+    return 1;
+  }
+  return 0;
+}
+
 /*
 ** This routine is called to do the work of a DROP TABLE statement.
 ** pName is the name of the table to be dropped.
@@ -109567,8 +112139,7 @@ SQLITE_PRIVATE void sqlite3DropTable(Parse *pParse, SrcList *pName, int isView,
     }
   }
 #endif
-  if( sqlite3StrNICmp(pTab->zName, "sqlite_", 7)==0 
-    && sqlite3StrNICmp(pTab->zName, "sqlite_stat", 11)!=0 ){
+  if( tableMayNotBeDropped(db, pTab) ){
     sqlite3ErrorMsg(pParse, "table %s may not be dropped", pTab->zName);
     goto exit_drop_table;
   }
@@ -109660,7 +112231,7 @@ SQLITE_PRIVATE void sqlite3CreateForeignKey(
   nByte = sizeof(*pFKey) + (nCol-1)*sizeof(pFKey->aCol[0]) + pTo->n + 1;
   if( pToCol ){
     for(i=0; i<pToCol->nExpr; i++){
-      nByte += sqlite3Strlen30(pToCol->a[i].zName) + 1;
+      nByte += sqlite3Strlen30(pToCol->a[i].zEName) + 1;
     }
   }
   pFKey = sqlite3DbMallocZero(db, nByte );
@@ -109685,7 +112256,7 @@ SQLITE_PRIVATE void sqlite3CreateForeignKey(
     for(i=0; i<nCol; i++){
       int j;
       for(j=0; j<p->nCol; j++){
-        if( sqlite3StrICmp(p->aCol[j].zName, pFromCol->a[i].zName)==0 ){
+        if( sqlite3StrICmp(p->aCol[j].zName, pFromCol->a[i].zEName)==0 ){
           pFKey->aCol[i].iFrom = j;
           break;
         }
@@ -109693,22 +112264,22 @@ SQLITE_PRIVATE void sqlite3CreateForeignKey(
       if( j>=p->nCol ){
         sqlite3ErrorMsg(pParse, 
           "unknown column \"%s\" in foreign key definition", 
-          pFromCol->a[i].zName);
+          pFromCol->a[i].zEName);
         goto fk_end;
       }
       if( IN_RENAME_OBJECT ){
-        sqlite3RenameTokenRemap(pParse, &pFKey->aCol[i], pFromCol->a[i].zName);
+        sqlite3RenameTokenRemap(pParse, &pFKey->aCol[i], pFromCol->a[i].zEName);
       }
     }
   }
   if( pToCol ){
     for(i=0; i<nCol; i++){
-      int n = sqlite3Strlen30(pToCol->a[i].zName);
+      int n = sqlite3Strlen30(pToCol->a[i].zEName);
       pFKey->aCol[i].zCol = z;
       if( IN_RENAME_OBJECT ){
-        sqlite3RenameTokenRemap(pParse, z, pToCol->a[i].zName);
+        sqlite3RenameTokenRemap(pParse, z, pToCol->a[i].zEName);
       }
-      memcpy(z, pToCol->a[i].zName, n);
+      memcpy(z, pToCol->a[i].zEName, n);
       z[n] = 0;
       z += n+1;
     }
@@ -109838,10 +112409,27 @@ static void sqlite3RefillIndex(Parse *pParse, Index *pIndex, int memRootPage){
     sqlite3UniqueConstraint(pParse, OE_Abort, pIndex);
     sqlite3VdbeJumpHere(v, j2);
   }else{
+    /* Most CREATE INDEX and REINDEX statements that are not UNIQUE can not
+    ** abort. The exception is if one of the indexed expressions contains a
+    ** user function that throws an exception when it is evaluated. But the
+    ** overhead of adding a statement journal to a CREATE INDEX statement is
+    ** very small (since most of the pages written do not contain content that
+    ** needs to be restored if the statement aborts), so we call 
+    ** sqlite3MayAbort() for all CREATE INDEX statements.  */
+    sqlite3MayAbort(pParse);
     addr2 = sqlite3VdbeCurrentAddr(v);
   }
   sqlite3VdbeAddOp3(v, OP_SorterData, iSorter, regRecord, iIdx);
-  sqlite3VdbeAddOp1(v, OP_SeekEnd, iIdx);
+  if( !pIndex->bAscKeyBug ){
+    /* This OP_SeekEnd opcode makes index insert for a REINDEX go much
+    ** faster by avoiding unnecessary seeks.  But the optimization does
+    ** not work for UNIQUE constraint indexes on WITHOUT ROWID tables
+    ** with DESC primary keys, since those indexes have there keys in
+    ** a different order from the main table.
+    ** See ticket: https://www.sqlite.org/src/info/bba7b69f9849b5bf
+    */
+    sqlite3VdbeAddOp1(v, OP_SeekEnd, iIdx);
+  }
   sqlite3VdbeAddOp2(v, OP_IdxInsert, iIdx, regRecord);
   sqlite3VdbeChangeP5(v, OPFLAG_USESEEKRESULT);
   sqlite3ReleaseTempReg(pParse, regRecord);
@@ -109888,6 +112476,27 @@ SQLITE_PRIVATE Index *sqlite3AllocateIndexObject(
   return p;
 }
 
+/*
+** If expression list pList contains an expression that was parsed with
+** an explicit "NULLS FIRST" or "NULLS LAST" clause, leave an error in
+** pParse and return non-zero. Otherwise, return zero.
+*/
+SQLITE_PRIVATE int sqlite3HasExplicitNulls(Parse *pParse, ExprList *pList){
+  if( pList ){
+    int i;
+    for(i=0; i<pList->nExpr; i++){
+      if( pList->a[i].bNulls ){
+        u8 sf = pList->a[i].sortFlags;
+        sqlite3ErrorMsg(pParse, "unsupported use of NULLS %s", 
+            (sf==0 || sf==3) ? "FIRST" : "LAST"
+        );
+        return 1;
+      }
+    }
+  }
+  return 0;
+}
+
 /*
 ** Create a new index for an SQL table.  pName1.pName2 is the name of the index 
 ** and pTblList is the name of the table that is to be indexed.  Both will 
@@ -109939,6 +112548,9 @@ SQLITE_PRIVATE void sqlite3CreateIndex(
   if( SQLITE_OK!=sqlite3ReadSchema(pParse) ){
     goto exit_create_index;
   }
+  if( sqlite3HasExplicitNulls(pParse, pList) ){
+    goto exit_create_index;
+  }
 
   /*
   ** Find the table that is to be indexed.  Return early if not found.
@@ -110037,7 +112649,7 @@ SQLITE_PRIVATE void sqlite3CreateIndex(
     zName = sqlite3NameFromToken(db, pName);
     if( zName==0 ) goto exit_create_index;
     assert( pName->z!=0 );
-    if( SQLITE_OK!=sqlite3CheckObjectName(pParse, zName) ){
+    if( SQLITE_OK!=sqlite3CheckObjectName(pParse, zName,"index",pTab->zName) ){
       goto exit_create_index;
     }
     if( !IN_RENAME_OBJECT ){
@@ -110103,7 +112715,7 @@ SQLITE_PRIVATE void sqlite3CreateIndex(
               sqlite3ExprAlloc(db, TK_ID, &prevCol, 0));
     if( pList==0 ) goto exit_create_index;
     assert( pList->nExpr==1 );
-    sqlite3ExprListSetSortOrder(pList, sortOrder);
+    sqlite3ExprListSetSortOrder(pList, sortOrder, SQLITE_SO_UNDEFINED);
   }else{
     sqlite3ExprListCheckLength(pParse, pList, "index");
     if( pParse->nErr ) goto exit_create_index;
@@ -110198,8 +112810,13 @@ SQLITE_PRIVATE void sqlite3CreateIndex(
       assert( j<=0x7fff );
       if( j<0 ){
         j = pTab->iPKey;
-      }else if( pTab->aCol[j].notNull==0 ){
-        pIndex->uniqNotNull = 0;
+      }else{
+        if( pTab->aCol[j].notNull==0 ){
+          pIndex->uniqNotNull = 0;
+        }
+        if( pTab->aCol[j].colFlags & COLFLAG_VIRTUAL ){
+          pIndex->bHasVCol = 1;
+        }
       }
       pIndex->aiColumn[i] = (i16)j;
     }
@@ -110221,7 +112838,7 @@ SQLITE_PRIVATE void sqlite3CreateIndex(
       goto exit_create_index;
     }
     pIndex->azColl[i] = zColl;
-    requestedSortOrder = pListItem->sortOrder & sortOrderMask;
+    requestedSortOrder = pListItem->sortFlags & sortOrderMask;
     pIndex->aSortOrder[i] = (u8)requestedSortOrder;
   }
 
@@ -110233,9 +112850,10 @@ SQLITE_PRIVATE void sqlite3CreateIndex(
     for(j=0; j<pPk->nKeyCol; j++){
       int x = pPk->aiColumn[j];
       assert( x>=0 );
-      if( hasColumn(pIndex->aiColumn, pIndex->nKeyCol, x) ){
+      if( isDupColumn(pIndex, pIndex->nKeyCol, pPk, j) ){
         pIndex->nColumn--; 
       }else{
+        testcase( hasColumn(pIndex->aiColumn,pIndex->nKeyCol,x) );
         pIndex->aiColumn[i] = x;
         pIndex->azColl[i] = pPk->azColl[j];
         pIndex->aSortOrder[i] = pPk->aSortOrder[j];
@@ -110253,13 +112871,13 @@ SQLITE_PRIVATE void sqlite3CreateIndex(
   /* If this index contains every column of its table, then mark
   ** it as a covering index */
   assert( HasRowid(pTab) 
-      || pTab->iPKey<0 || sqlite3ColumnOfIndex(pIndex, pTab->iPKey)>=0 );
+      || pTab->iPKey<0 || sqlite3TableColumnToIndex(pIndex, pTab->iPKey)>=0 );
   recomputeColumnsNotIndexed(pIndex);
   if( pTblName!=0 && pIndex->nColumn>=pTab->nCol ){
     pIndex->isCovering = 1;
     for(j=0; j<pTab->nCol; j++){
       if( j==pTab->iPKey ) continue;
-      if( sqlite3ColumnOfIndex(pIndex,j)>=0 ) continue;
+      if( sqlite3TableColumnToIndex(pIndex,j)>=0 ) continue;
       pIndex->isCovering = 0;
       break;
     }
@@ -110395,6 +113013,7 @@ SQLITE_PRIVATE void sqlite3CreateIndex(
       /* Gather the complete text of the CREATE INDEX statement into
       ** the zStmt variable
       */
+      assert( pName!=0 || pStart==0 );
       if( pStart ){
         int n = (int)(pParse->sLastToken.z - pName->z) + pParse->sLastToken.n;
         if( pName->z[n-1]==';' ) n--;
@@ -110433,26 +113052,9 @@ SQLITE_PRIVATE void sqlite3CreateIndex(
       sqlite3VdbeJumpHere(v, pIndex->tnum);
     }
   }
-
-  /* When adding an index to the list of indices for a table, make
-  ** sure all indices labeled OE_Replace come after all those labeled
-  ** OE_Ignore.  This is necessary for the correct constraint check
-  ** processing (in sqlite3GenerateConstraintChecks()) as part of
-  ** UPDATE and INSERT statements.  
-  */
   if( db->init.busy || pTblName==0 ){
-    if( onError!=OE_Replace || pTab->pIndex==0
-         || pTab->pIndex->onError==OE_Replace){
-      pIndex->pNext = pTab->pIndex;
-      pTab->pIndex = pIndex;
-    }else{
-      Index *pOther = pTab->pIndex;
-      while( pOther->pNext && pOther->pNext->onError!=OE_Replace ){
-        pOther = pOther->pNext;
-      }
-      pIndex->pNext = pOther->pNext;
-      pOther->pNext = pIndex;
-    }
+    pIndex->pNext = pTab->pIndex;
+    pTab->pIndex = pIndex;
     pIndex = 0;
   }
   else if( IN_RENAME_OBJECT ){
@@ -110464,6 +113066,21 @@ SQLITE_PRIVATE void sqlite3CreateIndex(
   /* Clean up before exiting */
 exit_create_index:
   if( pIndex ) sqlite3FreeIndex(db, pIndex);
+  if( pTab ){  /* Ensure all REPLACE indexes are at the end of the list */
+    Index **ppFrom = &pTab->pIndex;
+    Index *pThis;
+    for(ppFrom=&pTab->pIndex; (pThis = *ppFrom)!=0; ppFrom=&pThis->pNext){
+      Index *pNext;
+      if( pThis->onError!=OE_Replace ) continue;
+      while( (pNext = pThis->pNext)!=0 && pNext->onError!=OE_Replace ){
+        *ppFrom = pNext;
+        pThis->pNext = pNext->pNext;
+        pNext->pNext = pThis;
+        ppFrom = &pNext->pNext;
+      }
+      break;
+    }
+  }
   sqlite3ExprDelete(db, pPIWhere);
   sqlite3ExprListDelete(db, pList);
   sqlite3SrcListDelete(db, pTblName);
@@ -111437,7 +114054,8 @@ SQLITE_PRIVATE KeyInfo *sqlite3KeyInfoOfIndex(Parse *pParse, Index *pIdx){
       const char *zColl = pIdx->azColl[i];
       pKey->aColl[i] = zColl==sqlite3StrBINARY ? 0 :
                         sqlite3LocateCollSeq(pParse, zColl);
-      pKey->aSortOrder[i] = pIdx->aSortOrder[i];
+      pKey->aSortFlags[i] = pIdx->aSortOrder[i];
+      assert( 0==(pKey->aSortFlags[i] & KEYINFO_ORDER_BIGNULL) );
     }
     if( pParse->nErr ){
       assert( pParse->rc==SQLITE_ERROR_MISSING_COLLSEQ );
@@ -111597,51 +114215,6 @@ static int synthCollSeq(sqlite3 *db, CollSeq *pColl){
   return SQLITE_ERROR;
 }
 
-/*
-** This function is responsible for invoking the collation factory callback
-** or substituting a collation sequence of a different encoding when the
-** requested collation sequence is not available in the desired encoding.
-** 
-** If it is not NULL, then pColl must point to the database native encoding 
-** collation sequence with name zName, length nName.
-**
-** The return value is either the collation sequence to be used in database
-** db for collation type name zName, length nName, or NULL, if no collation
-** sequence can be found.  If no collation is found, leave an error message.
-**
-** See also: sqlite3LocateCollSeq(), sqlite3FindCollSeq()
-*/
-SQLITE_PRIVATE CollSeq *sqlite3GetCollSeq(
-  Parse *pParse,        /* Parsing context */
-  u8 enc,               /* The desired encoding for the collating sequence */
-  CollSeq *pColl,       /* Collating sequence with native encoding, or NULL */
-  const char *zName     /* Collating sequence name */
-){
-  CollSeq *p;
-  sqlite3 *db = pParse->db;
-
-  p = pColl;
-  if( !p ){
-    p = sqlite3FindCollSeq(db, enc, zName, 0);
-  }
-  if( !p || !p->xCmp ){
-    /* No collation sequence of this type for this encoding is registered.
-    ** Call the collation factory to see if it can supply us with one.
-    */
-    callCollNeeded(db, enc, zName);
-    p = sqlite3FindCollSeq(db, enc, zName, 0);
-  }
-  if( p && !p->xCmp && synthCollSeq(db, p) ){
-    p = 0;
-  }
-  assert( !p || p->xCmp );
-  if( p==0 ){
-    sqlite3ErrorMsg(pParse, "no such collation sequence: %s", zName);
-    pParse->rc = SQLITE_ERROR_MISSING_COLLSEQ;
-  }
-  return p;
-}
-
 /*
 ** This routine is called on a collation sequence before it is used to
 ** check that it is defined. An undefined collation sequence exists when
@@ -111734,10 +114307,10 @@ static CollSeq *findCollSeqEntry(
 ** See also: sqlite3LocateCollSeq(), sqlite3GetCollSeq()
 */
 SQLITE_PRIVATE CollSeq *sqlite3FindCollSeq(
-  sqlite3 *db,
-  u8 enc,
-  const char *zName,
-  int create
+  sqlite3 *db,          /* Database connection to search */
+  u8 enc,               /* Desired text encoding */
+  const char *zName,    /* Name of the collating sequence.  Might be NULL */
+  int create            /* True to create CollSeq if doesn't already exist */
 ){
   CollSeq *pColl;
   if( zName ){
@@ -111751,6 +114324,85 @@ SQLITE_PRIVATE CollSeq *sqlite3FindCollSeq(
   return pColl;
 }
 
+/*
+** This function is responsible for invoking the collation factory callback
+** or substituting a collation sequence of a different encoding when the
+** requested collation sequence is not available in the desired encoding.
+** 
+** If it is not NULL, then pColl must point to the database native encoding 
+** collation sequence with name zName, length nName.
+**
+** The return value is either the collation sequence to be used in database
+** db for collation type name zName, length nName, or NULL, if no collation
+** sequence can be found.  If no collation is found, leave an error message.
+**
+** See also: sqlite3LocateCollSeq(), sqlite3FindCollSeq()
+*/
+SQLITE_PRIVATE CollSeq *sqlite3GetCollSeq(
+  Parse *pParse,        /* Parsing context */
+  u8 enc,               /* The desired encoding for the collating sequence */
+  CollSeq *pColl,       /* Collating sequence with native encoding, or NULL */
+  const char *zName     /* Collating sequence name */
+){
+  CollSeq *p;
+  sqlite3 *db = pParse->db;
+
+  p = pColl;
+  if( !p ){
+    p = sqlite3FindCollSeq(db, enc, zName, 0);
+  }
+  if( !p || !p->xCmp ){
+    /* No collation sequence of this type for this encoding is registered.
+    ** Call the collation factory to see if it can supply us with one.
+    */
+    callCollNeeded(db, enc, zName);
+    p = sqlite3FindCollSeq(db, enc, zName, 0);
+  }
+  if( p && !p->xCmp && synthCollSeq(db, p) ){
+    p = 0;
+  }
+  assert( !p || p->xCmp );
+  if( p==0 ){
+    sqlite3ErrorMsg(pParse, "no such collation sequence: %s", zName);
+    pParse->rc = SQLITE_ERROR_MISSING_COLLSEQ;
+  }
+  return p;
+}
+
+/*
+** This function returns the collation sequence for database native text
+** encoding identified by the string zName.
+**
+** If the requested collation sequence is not available, or not available
+** in the database native encoding, the collation factory is invoked to
+** request it. If the collation factory does not supply such a sequence,
+** and the sequence is available in another text encoding, then that is
+** returned instead.
+**
+** If no versions of the requested collations sequence are available, or
+** another error occurs, NULL is returned and an error message written into
+** pParse.
+**
+** This routine is a wrapper around sqlite3FindCollSeq().  This routine
+** invokes the collation factory if the named collation cannot be found
+** and generates an error message.
+**
+** See also: sqlite3FindCollSeq(), sqlite3GetCollSeq()
+*/
+SQLITE_PRIVATE CollSeq *sqlite3LocateCollSeq(Parse *pParse, const char *zName){
+  sqlite3 *db = pParse->db;
+  u8 enc = ENC(db);
+  u8 initbusy = db->init.busy;
+  CollSeq *pColl;
+
+  pColl = sqlite3FindCollSeq(db, enc, zName, initbusy);
+  if( !initbusy && (!pColl || !pColl->xCmp) ){
+    pColl = sqlite3GetCollSeq(pParse, enc, pColl, zName);
+  }
+
+  return pColl;
+}
+
 /* During the search for the best function definition, this procedure
 ** is called to test how well the function passed as the first argument
 ** matches the request for a function with nArg arguments in a system
@@ -111786,12 +114438,13 @@ static int matchQuality(
   u8 enc          /* Desired text encoding */
 ){
   int match;
-
-  /* nArg of -2 is a special case */
-  if( nArg==(-2) ) return (p->xSFunc==0) ? 0 : FUNC_PERFECT_MATCH;
+  assert( p->nArg>=-1 );
 
   /* Wrong number of arguments means "no match" */
-  if( p->nArg!=nArg && p->nArg>=0 ) return 0;
+  if( p->nArg!=nArg ){
+    if( nArg==(-2) ) return (p->xSFunc==0) ? 0 : FUNC_PERFECT_MATCH;
+    if( p->nArg>=0 ) return 0;
+  }
 
   /* Give a better score to a function with a specific number of arguments
   ** than to function that accepts any number of arguments. */
@@ -112093,11 +114746,7 @@ static int tabIsReadOnly(Parse *pParse, Table *pTab){
     return sqlite3WritableSchema(db)==0 && pParse->nested==0;
   }
   assert( pTab->tabFlags & TF_Shadow );
-  return (db->flags & SQLITE_Defensive)!=0 
-#ifndef SQLITE_OMIT_VIRTUALTABLE
-          && db->pVtabCtx==0
-#endif
-          && db->nVdbeExec==0;
+  return sqlite3ReadOnlyShadowTables(db);
 }
 
 /*
@@ -112760,7 +115409,8 @@ SQLITE_PRIVATE void sqlite3GenerateRowDelete(
       testcase( mask!=0xffffffff && iCol==31 );
       testcase( mask!=0xffffffff && iCol==32 );
       if( mask==0xffffffff || (iCol<=31 && (mask & MASKBIT32(iCol))!=0) ){
-        sqlite3ExprCodeGetColumnOfTable(v, pTab, iDataCur, iCol, iOld+iCol+1);
+        int kk = sqlite3TableColumnToStorage(pTab, iCol);
+        sqlite3ExprCodeGetColumnOfTable(v, pTab, iDataCur, iCol, iOld+kk+1);
       }
     }
 
@@ -112940,6 +115590,8 @@ SQLITE_PRIVATE int sqlite3GenerateIndexKey(
       sqlite3ExprIfFalseDup(pParse, pIdx->pPartIdxWhere, *piPartIdxLabel, 
                             SQLITE_JUMPIFNULL);
       pParse->iSelfTab = 0;
+      pPrior = 0; /* Ticket a9efb42811fa41ee 2019-11-02;
+                  ** pPartIdxWhere may have corrupted regPrior registers */
     }else{
       *piPartIdxLabel = 0;
     }
@@ -113006,6 +115658,9 @@ SQLITE_PRIVATE void sqlite3ResolvePartIdxLabel(Parse *pParse, int iLabel){
 /* #include "sqliteInt.h" */
 /* #include <stdlib.h> */
 /* #include <assert.h> */
+#ifndef SQLITE_OMIT_FLOATING_POINT
+/* #include <math.h> */
+#endif
 /* #include "vdbeInt.h" */
 
 /*
@@ -113192,6 +115847,8 @@ static void instrFunc(
   int N = 1;
   int isText;
   unsigned char firstChar;
+  sqlite3_value *pC1 = 0;
+  sqlite3_value *pC2 = 0;
 
   UNUSED_PARAMETER(argc);
   typeHaystack = sqlite3_value_type(argv[0]);
@@ -113204,12 +115861,22 @@ static void instrFunc(
       zHaystack = sqlite3_value_blob(argv[0]);
       zNeedle = sqlite3_value_blob(argv[1]);
       isText = 0;
-    }else{
+    }else if( typeHaystack!=SQLITE_BLOB && typeNeedle!=SQLITE_BLOB ){
       zHaystack = sqlite3_value_text(argv[0]);
       zNeedle = sqlite3_value_text(argv[1]);
       isText = 1;
+    }else{
+      pC1 = sqlite3_value_dup(argv[0]);
+      zHaystack = sqlite3_value_text(pC1);
+      if( zHaystack==0 ) goto endInstrOOM;
+      nHaystack = sqlite3_value_bytes(pC1);
+      pC2 = sqlite3_value_dup(argv[1]);
+      zNeedle = sqlite3_value_text(pC2);
+      if( zNeedle==0 ) goto endInstrOOM;
+      nNeedle = sqlite3_value_bytes(pC2);
+      isText = 1;
     }
-    if( zNeedle==0 || (nHaystack && zHaystack==0) ) return;
+    if( zNeedle==0 || (nHaystack && zHaystack==0) ) goto endInstrOOM;
     firstChar = zNeedle[0];
     while( nNeedle<=nHaystack
        && (zHaystack[0]!=firstChar || memcmp(zHaystack, zNeedle, nNeedle)!=0)
@@ -113223,6 +115890,13 @@ static void instrFunc(
     if( nNeedle>nHaystack ) N = 0;
   }
   sqlite3_result_int(context, N);
+endInstr:
+  sqlite3_value_free(pC1);
+  sqlite3_value_free(pC2);
+  return;
+endInstrOOM:
+  sqlite3_result_error_nomem(context);
+  goto endInstr;
 }
 
 /*
@@ -113376,10 +116050,10 @@ static void roundFunc(sqlite3_context *context, int argc, sqlite3_value **argv){
   ** handle the rounding directly,
   ** otherwise use printf.
   */
-  if( n==0 && r>=0 && r<LARGEST_INT64-1 ){
-    r = (double)((sqlite_int64)(r+0.5));
-  }else if( n==0 && r<0 && (-r)<LARGEST_INT64-1 ){
-    r = -(double)((sqlite_int64)((-r)+0.5));
+  if( r<-4503599627370496.0 || r>+4503599627370496.0 ){
+    /* The value has no fractional part so there is nothing to round */
+  }else if( n==0 ){  
+    r = (double)((sqlite_int64)(r+(r<0?-0.5:+0.5)));
   }else{
     zBuf = sqlite3_mprintf("%.*f",n,r);
     if( zBuf==0 ){
@@ -113833,8 +116507,6 @@ static void likeFunc(
     return;
   }
 #endif
-  zB = sqlite3_value_text(argv[0]);
-  zA = sqlite3_value_text(argv[1]);
 
   /* Limit the length of the LIKE or GLOB pattern to avoid problems
   ** of deep recursion and N*N behavior in patternCompare().
@@ -113846,8 +116518,6 @@ static void likeFunc(
     sqlite3_result_error(context, "LIKE or GLOB pattern too complex", -1);
     return;
   }
-  assert( zB==sqlite3_value_text(argv[0]) );  /* Encoding did not change */
-
   if( argc==3 ){
     /* The escape character string must consist of a single UTF-8 character.
     ** Otherwise, return an error.
@@ -113863,6 +116533,8 @@ static void likeFunc(
   }else{
     escape = pInfo->matchSet;
   }
+  zB = sqlite3_value_text(argv[0]);
+  zA = sqlite3_value_text(argv[1]);
   if( zA && zB ){
 #ifdef SQLITE_TEST
     sqlite3_like_count++;
@@ -114788,39 +117460,24 @@ SQLITE_PRIVATE void sqlite3RegisterPerConnectionBuiltinFunctions(sqlite3 *db){
 }
 
 /*
-** Set the LIKEOPT flag on the 2-argument function with the given name.
-*/
-static void setLikeOptFlag(sqlite3 *db, const char *zName, u8 flagVal){
-  FuncDef *pDef;
-  pDef = sqlite3FindFunction(db, zName, 2, SQLITE_UTF8, 0);
-  if( ALWAYS(pDef) ){
-    pDef->funcFlags |= flagVal;
-  }
-  pDef = sqlite3FindFunction(db, zName, 3, SQLITE_UTF8, 0);
-  if( pDef ){
-    pDef->funcFlags |= flagVal;
-  }
-}
-
-/*
-** Register the built-in LIKE and GLOB functions.  The caseSensitive
+** Re-register the built-in LIKE functions.  The caseSensitive
 ** parameter determines whether or not the LIKE operator is case
-** sensitive.  GLOB is always case sensitive.
+** sensitive.
 */
 SQLITE_PRIVATE void sqlite3RegisterLikeFunctions(sqlite3 *db, int caseSensitive){
   struct compareInfo *pInfo;
+  int flags;
   if( caseSensitive ){
     pInfo = (struct compareInfo*)&likeInfoAlt;
+    flags = SQLITE_FUNC_LIKE | SQLITE_FUNC_CASE;
   }else{
     pInfo = (struct compareInfo*)&likeInfoNorm;
+    flags = SQLITE_FUNC_LIKE;
   }
   sqlite3CreateFunc(db, "like", 2, SQLITE_UTF8, pInfo, likeFunc, 0, 0, 0, 0, 0);
   sqlite3CreateFunc(db, "like", 3, SQLITE_UTF8, pInfo, likeFunc, 0, 0, 0, 0, 0);
-  sqlite3CreateFunc(db, "glob", 2, SQLITE_UTF8, 
-      (struct compareInfo*)&globInfo, likeFunc, 0, 0, 0, 0, 0);
-  setLikeOptFlag(db, "glob", SQLITE_FUNC_LIKE | SQLITE_FUNC_CASE);
-  setLikeOptFlag(db, "like", 
-      caseSensitive ? (SQLITE_FUNC_LIKE | SQLITE_FUNC_CASE) : SQLITE_FUNC_LIKE);
+  sqlite3FindFunction(db, "like", 2, SQLITE_UTF8, 0)->funcFlags |= flags;
+  sqlite3FindFunction(db, "like", 3, SQLITE_UTF8, 0)->funcFlags |= flags;
 }
 
 /*
@@ -114849,6 +117506,9 @@ SQLITE_PRIVATE int sqlite3IsLikeFunction(sqlite3 *db, Expr *pExpr, int *pIsNocas
   assert( !ExprHasProperty(pExpr, EP_xIsSelect) );
   nExpr = pExpr->x.pList->nExpr;
   pDef = sqlite3FindFunction(db, pExpr->u.zToken, nExpr, SQLITE_UTF8, 0);
+#ifdef SQLITE_ENABLE_UNKNOWN_SQL_FUNCTION
+  if( pDef==0 ) return 0;
+#endif
   if( NEVER(pDef==0) || (pDef->funcFlags & SQLITE_FUNC_LIKE)==0 ){
     return 0;
   }
@@ -114894,12 +117554,20 @@ SQLITE_PRIVATE void sqlite3RegisterBuiltinFunctions(void){
   ** For peak efficiency, put the most frequently used function last.
   */
   static FuncDef aBuiltinFunc[] = {
+/***** Functions only available with SQLITE_TESTCTRL_INTERNAL_FUNCTIONS *****/
+    TEST_FUNC(implies_nonnull_row, 2, INLINEFUNC_implies_nonnull_row, 0),
+    TEST_FUNC(expr_compare,        2, INLINEFUNC_expr_compare,        0),
+    TEST_FUNC(expr_implies_expr,   2, INLINEFUNC_expr_implies_expr,   0),
+#ifdef SQLITE_DEBUG
+    TEST_FUNC(affinity,          1, INLINEFUNC_affinity, 0),
+#endif
+/***** Regular functions *****/
 #ifdef SQLITE_SOUNDEX
     FUNCTION(soundex,            1, 0, 0, soundexFunc      ),
 #endif
 #ifndef SQLITE_OMIT_LOAD_EXTENSION
-    VFUNCTION(load_extension,    1, 0, 0, loadExt          ),
-    VFUNCTION(load_extension,    2, 0, 0, loadExt          ),
+    SFUNCTION(load_extension,    1, 0, 0, loadExt          ),
+    SFUNCTION(load_extension,    2, 0, 0, loadExt          ),
 #endif
 #if SQLITE_USER_AUTHENTICATION
     FUNCTION(sqlite_crypt,       2, 0, 0, sqlite3CryptFunc ),
@@ -114908,12 +117576,9 @@ SQLITE_PRIVATE void sqlite3RegisterBuiltinFunctions(void){
     DFUNCTION(sqlite_compileoption_used,1, 0, 0, compileoptionusedFunc  ),
     DFUNCTION(sqlite_compileoption_get, 1, 0, 0, compileoptiongetFunc  ),
 #endif /* SQLITE_OMIT_COMPILEOPTION_DIAGS */
-    FUNCTION2(unlikely,          1, 0, 0, noopFunc,  SQLITE_FUNC_UNLIKELY),
-    FUNCTION2(likelihood,        2, 0, 0, noopFunc,  SQLITE_FUNC_UNLIKELY),
-    FUNCTION2(likely,            1, 0, 0, noopFunc,  SQLITE_FUNC_UNLIKELY),
-#ifdef SQLITE_DEBUG
-    FUNCTION2(affinity,          1, 0, 0, noopFunc,  SQLITE_FUNC_AFFINITY),
-#endif
+    INLINE_FUNC(unlikely,        1, INLINEFUNC_unlikely, SQLITE_FUNC_UNLIKELY),
+    INLINE_FUNC(likelihood,      2, INLINEFUNC_unlikely, SQLITE_FUNC_UNLIKELY),
+    INLINE_FUNC(likely,          1, INLINEFUNC_unlikely, SQLITE_FUNC_UNLIKELY),
 #ifdef SQLITE_ENABLE_OFFSET_SQL_FUNC
     FUNCTION2(sqlite_offset,     1, 0, 0, noopFunc,  SQLITE_FUNC_OFFSET|
                                                      SQLITE_FUNC_TYPEOF),
@@ -114946,7 +117611,7 @@ SQLITE_PRIVATE void sqlite3RegisterBuiltinFunctions(void){
     FUNCTION(upper,              1, 0, 0, upperFunc        ),
     FUNCTION(lower,              1, 0, 0, lowerFunc        ),
     FUNCTION(hex,                1, 0, 0, hexFunc          ),
-    FUNCTION2(ifnull,            2, 0, 0, noopFunc,  SQLITE_FUNC_COALESCE),
+    INLINE_FUNC(ifnull,          2, INLINEFUNC_coalesce, SQLITE_FUNC_COALESCE),
     VFUNCTION(random,            0, 0, 0, randomFunc       ),
     VFUNCTION(randomblob,        1, 0, 0, randomBlob       ),
     FUNCTION(nullif,             2, 0, 1, nullifFunc       ),
@@ -114986,15 +117651,12 @@ SQLITE_PRIVATE void sqlite3RegisterBuiltinFunctions(void){
 #endif
     FUNCTION(coalesce,           1, 0, 0, 0                ),
     FUNCTION(coalesce,           0, 0, 0, 0                ),
-    FUNCTION2(coalesce,         -1, 0, 0, noopFunc,  SQLITE_FUNC_COALESCE),
+    INLINE_FUNC(coalesce,       -1, INLINEFUNC_coalesce, SQLITE_FUNC_COALESCE),
   };
 #ifndef SQLITE_OMIT_ALTERTABLE
   sqlite3AlterFunctions();
 #endif
   sqlite3WindowFunctions();
-#if defined(SQLITE_ENABLE_STAT3) || defined(SQLITE_ENABLE_STAT4)
-  sqlite3AnalyzeFunctions();
-#endif
   sqlite3RegisterDateTimeFunctions();
   sqlite3InsertBuiltinFuncs(aBuiltinFunc, ArraySize(aBuiltinFunc));
 
@@ -115368,7 +118030,7 @@ static void fkLookupParent(
     VdbeCoverage(v);
   }
   for(i=0; i<pFKey->nCol; i++){
-    int iReg = aiCol[i] + regData + 1;
+    int iReg = sqlite3TableColumnToStorage(pFKey->pFrom,aiCol[i]) + regData + 1;
     sqlite3VdbeAddOp2(v, OP_IsNull, iReg, iOk); VdbeCoverage(v);
   }
 
@@ -115384,7 +118046,8 @@ static void fkLookupParent(
       ** is no matching parent key. Before using MustBeInt, make a copy of
       ** the value. Otherwise, the value inserted into the child key column
       ** will have INTEGER affinity applied to it, which may not be correct.  */
-      sqlite3VdbeAddOp2(v, OP_SCopy, aiCol[0]+1+regData, regTemp);
+      sqlite3VdbeAddOp2(v, OP_SCopy, 
+        sqlite3TableColumnToStorage(pFKey->pFrom,aiCol[0])+1+regData, regTemp);
       iMustBeInt = sqlite3VdbeAddOp2(v, OP_MustBeInt, regTemp, 0);
       VdbeCoverage(v);
   
@@ -115411,7 +118074,9 @@ static void fkLookupParent(
       sqlite3VdbeAddOp3(v, OP_OpenRead, iCur, pIdx->tnum, iDb);
       sqlite3VdbeSetP4KeyInfo(pParse, pIdx);
       for(i=0; i<nCol; i++){
-        sqlite3VdbeAddOp2(v, OP_Copy, aiCol[i]+1+regData, regTemp+i);
+        sqlite3VdbeAddOp2(v, OP_Copy, 
+               sqlite3TableColumnToStorage(pFKey->pFrom, aiCol[i])+1+regData,
+               regTemp+i);
       }
   
       /* If the parent table is the same as the child table, and we are about
@@ -115427,8 +118092,11 @@ static void fkLookupParent(
       if( pTab==pFKey->pFrom && nIncr==1 ){
         int iJump = sqlite3VdbeCurrentAddr(v) + nCol + 1;
         for(i=0; i<nCol; i++){
-          int iChild = aiCol[i]+1+regData;
-          int iParent = pIdx->aiColumn[i]+1+regData;
+          int iChild = sqlite3TableColumnToStorage(pFKey->pFrom,aiCol[i])
+                              +1+regData;
+          int iParent = 1+regData;
+          iParent += sqlite3TableColumnToStorage(pIdx->pTable,
+                                                 pIdx->aiColumn[i]);
           assert( pIdx->aiColumn[i]>=0 );
           assert( aiCol[i]!=pTab->iPKey );
           if( pIdx->aiColumn[i]==pTab->iPKey ){
@@ -115496,14 +118164,14 @@ static Expr *exprTableRegister(
   if( pExpr ){
     if( iCol>=0 && iCol!=pTab->iPKey ){
       pCol = &pTab->aCol[iCol];
-      pExpr->iTable = regBase + iCol + 1;
-      pExpr->affinity = pCol->affinity;
+      pExpr->iTable = regBase + sqlite3TableColumnToStorage(pTab,iCol) + 1;
+      pExpr->affExpr = pCol->affinity;
       zColl = pCol->zColl;
       if( zColl==0 ) zColl = db->pDfltColl->zName;
       pExpr = sqlite3ExprAddCollateString(pParse, pExpr, zColl);
     }else{
       pExpr->iTable = regBase;
-      pExpr->affinity = SQLITE_AFF_INTEGER;
+      pExpr->affExpr = SQLITE_AFF_INTEGER;
     }
   }
   return pExpr;
@@ -115610,7 +118278,7 @@ static void fkScanChildren(
     zCol = pFKey->pFrom->aCol[iCol].zName;
     pRight = sqlite3Expr(db, TK_ID, zCol);
     pEq = sqlite3PExpr(pParse, TK_EQ, pLeft, pRight);
-    pWhere = sqlite3ExprAnd(db, pWhere, pEq);
+    pWhere = sqlite3ExprAnd(pParse, pWhere, pEq);
   }
 
   /* If the child table is the same as the parent table, then add terms
@@ -115644,11 +118312,11 @@ static void fkScanChildren(
         pLeft = exprTableRegister(pParse, pTab, regData, iCol);
         pRight = sqlite3Expr(db, TK_ID, pTab->aCol[iCol].zName);
         pEq = sqlite3PExpr(pParse, TK_IS, pLeft, pRight);
-        pAll = sqlite3ExprAnd(db, pAll, pEq);
+        pAll = sqlite3ExprAnd(pParse, pAll, pEq);
       }
       pNe = sqlite3PExpr(pParse, TK_NOT, pAll, 0);
     }
-    pWhere = sqlite3ExprAnd(db, pWhere, pNe);
+    pWhere = sqlite3ExprAnd(pParse, pWhere, pNe);
   }
 
   /* Resolve the references in the WHERE clause. */
@@ -115945,7 +118613,9 @@ SQLITE_PRIVATE void sqlite3FkCheck(
         Vdbe *v = sqlite3GetVdbe(pParse);
         int iJump = sqlite3VdbeCurrentAddr(v) + pFKey->nCol + 1;
         for(i=0; i<pFKey->nCol; i++){
-          int iReg = pFKey->aCol[i].iFrom + regOld + 1;
+          int iFromCol, iReg;
+          iFromCol = pFKey->aCol[i].iFrom;
+          iReg = sqlite3TableColumnToStorage(pFKey->pFrom,iFromCol) + regOld+1;
           sqlite3VdbeAddOp2(v, OP_IsNull, iReg, iJump); VdbeCoverage(v);
         }
         sqlite3VdbeAddOp2(v, OP_FkCounter, pFKey->isDeferred, -1);
@@ -116254,7 +118924,7 @@ static Trigger *fkActionTrigger(
             sqlite3ExprAlloc(db, TK_ID, &tToCol, 0)),
           sqlite3ExprAlloc(db, TK_ID, &tFromCol, 0)
       );
-      pWhere = sqlite3ExprAnd(db, pWhere, pEq);
+      pWhere = sqlite3ExprAnd(pParse, pWhere, pEq);
 
       /* For ON UPDATE, construct the next term of the WHEN clause.
       ** The final WHEN clause will be like this:
@@ -116270,7 +118940,7 @@ static Trigger *fkActionTrigger(
               sqlite3ExprAlloc(db, TK_ID, &tNew, 0),
               sqlite3ExprAlloc(db, TK_ID, &tToCol, 0))
             );
-        pWhen = sqlite3ExprAnd(db, pWhen, pEq);
+        pWhen = sqlite3ExprAnd(pParse, pWhen, pEq);
       }
   
       if( action!=OE_Restrict && (action!=OE_Cascade || pChanges) ){
@@ -116280,7 +118950,15 @@ static Trigger *fkActionTrigger(
             sqlite3ExprAlloc(db, TK_ID, &tNew, 0),
             sqlite3ExprAlloc(db, TK_ID, &tToCol, 0));
         }else if( action==OE_SetDflt ){
-          Expr *pDflt = pFKey->pFrom->aCol[iFromCol].pDflt;
+          Column *pCol = pFKey->pFrom->aCol + iFromCol;
+          Expr *pDflt;
+          if( pCol->colFlags & COLFLAG_GENERATED ){
+            testcase( pCol->colFlags & COLFLAG_VIRTUAL );
+            testcase( pCol->colFlags & COLFLAG_STORED );
+            pDflt = 0;
+          }else{
+            pDflt = pCol->pDflt;
+          }
           if( pDflt ){
             pNew = sqlite3ExprDup(db, pDflt, 0);
           }else{
@@ -116306,7 +118984,7 @@ static Trigger *fkActionTrigger(
       tFrom.n = nFrom;
       pRaise = sqlite3Expr(db, TK_RAISE, "FOREIGN KEY constraint failed");
       if( pRaise ){
-        pRaise->affinity = OE_Abort;
+        pRaise->affExpr = OE_Abort;
       }
       pSelect = sqlite3SelectNew(pParse, 
           sqlite3ExprListAppend(pParse, 0, pRaise),
@@ -116318,7 +118996,7 @@ static Trigger *fkActionTrigger(
     }
 
     /* Disable lookaside memory allocation */
-    db->lookaside.bDisable++;
+    DisableLookaside;
 
     pTrigger = (Trigger *)sqlite3DbMallocZero(db, 
         sizeof(Trigger) +         /* struct Trigger */
@@ -116340,7 +119018,7 @@ static Trigger *fkActionTrigger(
     }
 
     /* Re-enable the lookaside buffer, if it was disabled earlier. */
-    db->lookaside.bDisable--;
+    EnableLookaside;
 
     sqlite3ExprDelete(db, pWhere);
     sqlite3ExprDelete(db, pWhen);
@@ -116351,6 +119029,7 @@ static Trigger *fkActionTrigger(
       return 0;
     }
     assert( pStep!=0 );
+    assert( pTrigger!=0 );
 
     switch( action ){
       case OE_Restrict:
@@ -116490,7 +119169,7 @@ SQLITE_PRIVATE void sqlite3OpenTable(
   sqlite3TableLock(pParse, iDb, pTab->tnum, 
                    (opcode==OP_OpenWrite)?1:0, pTab->zName);
   if( HasRowid(pTab) ){
-    sqlite3VdbeAddOp4Int(v, opcode, iCur, pTab->tnum, iDb, pTab->nCol);
+    sqlite3VdbeAddOp4Int(v, opcode, iCur, pTab->tnum, iDb, pTab->nNVCol);
     VdbeComment((v, "%s", pTab->zName));
   }else{
     Index *pPk = sqlite3PrimaryKeyIndex(pTab);
@@ -116541,18 +119220,19 @@ SQLITE_PRIVATE const char *sqlite3IndexAffinityStr(sqlite3 *db, Index *pIdx){
     }
     for(n=0; n<pIdx->nColumn; n++){
       i16 x = pIdx->aiColumn[n];
+      char aff;
       if( x>=0 ){
-        pIdx->zColAff[n] = pTab->aCol[x].affinity;
+        aff = pTab->aCol[x].affinity;
       }else if( x==XN_ROWID ){
-        pIdx->zColAff[n] = SQLITE_AFF_INTEGER;
+        aff = SQLITE_AFF_INTEGER;
       }else{
-        char aff;
         assert( x==XN_EXPR );
         assert( pIdx->aColExpr!=0 );
         aff = sqlite3ExprAffinity(pIdx->aColExpr->a[n].pExpr);
-        if( aff==0 ) aff = SQLITE_AFF_BLOB;
-        pIdx->zColAff[n] = aff;
       }
+      if( aff<SQLITE_AFF_BLOB ) aff = SQLITE_AFF_BLOB;
+      if( aff>SQLITE_AFF_NUMERIC) aff = SQLITE_AFF_NUMERIC;
+      pIdx->zColAff[n] = aff;
     }
     pIdx->zColAff[n] = 0;
   }
@@ -116581,7 +119261,7 @@ SQLITE_PRIVATE const char *sqlite3IndexAffinityStr(sqlite3 *db, Index *pIdx){
 **  'E'            REAL
 */
 SQLITE_PRIVATE void sqlite3TableAffinity(Vdbe *v, Table *pTab, int iReg){
-  int i;
+  int i, j;
   char *zColAff = pTab->zColAff;
   if( zColAff==0 ){
     sqlite3 *db = sqlite3VdbeDb(v);
@@ -116591,12 +119271,15 @@ SQLITE_PRIVATE void sqlite3TableAffinity(Vdbe *v, Table *pTab, int iReg){
       return;
     }
 
-    for(i=0; i<pTab->nCol; i++){
-      zColAff[i] = pTab->aCol[i].affinity;
+    for(i=j=0; i<pTab->nCol; i++){
+      assert( pTab->aCol[i].affinity!=0 );
+      if( (pTab->aCol[i].colFlags & COLFLAG_VIRTUAL)==0 ){
+        zColAff[j++] = pTab->aCol[i].affinity;
+      }
     }
     do{
-      zColAff[i--] = 0;
-    }while( i>=0 && zColAff[i]==SQLITE_AFF_BLOB );
+      zColAff[j--] = 0;
+    }while( j>=0 && zColAff[j]<=SQLITE_AFF_BLOB );
     pTab->zColAff = zColAff;
   }
   assert( zColAff!=0 );
@@ -116650,6 +119333,119 @@ static int readsTable(Parse *p, int iDb, Table *pTab){
   return 0;
 }
 
+/* This walker callback will compute the union of colFlags flags for all
+** referenced columns in a CHECK constraint or generated column expression.
+*/
+static int exprColumnFlagUnion(Walker *pWalker, Expr *pExpr){
+  if( pExpr->op==TK_COLUMN && pExpr->iColumn>=0 ){
+    assert( pExpr->iColumn < pWalker->u.pTab->nCol );
+    pWalker->eCode |= pWalker->u.pTab->aCol[pExpr->iColumn].colFlags;
+  }
+  return WRC_Continue;
+}
+
+#ifndef SQLITE_OMIT_GENERATED_COLUMNS
+/*
+** All regular columns for table pTab have been puts into registers
+** starting with iRegStore.  The registers that correspond to STORED
+** or VIRTUAL columns have not yet been initialized.  This routine goes
+** back and computes the values for those columns based on the previously
+** computed normal columns.
+*/
+SQLITE_PRIVATE void sqlite3ComputeGeneratedColumns(
+  Parse *pParse,    /* Parsing context */
+  int iRegStore,    /* Register holding the first column */
+  Table *pTab       /* The table */
+){
+  int i;
+  Walker w;
+  Column *pRedo;
+  int eProgress;
+  VdbeOp *pOp;
+
+  assert( pTab->tabFlags & TF_HasGenerated );
+  testcase( pTab->tabFlags & TF_HasVirtual );
+  testcase( pTab->tabFlags & TF_HasStored );
+
+  /* Before computing generated columns, first go through and make sure
+  ** that appropriate affinity has been applied to the regular columns
+  */
+  sqlite3TableAffinity(pParse->pVdbe, pTab, iRegStore);
+  if( (pTab->tabFlags & TF_HasStored)!=0
+   && (pOp = sqlite3VdbeGetOp(pParse->pVdbe,-1))->opcode==OP_Affinity
+  ){
+    /* Change the OP_Affinity argument to '@' (NONE) for all stored
+    ** columns.  '@' is the no-op affinity and those columns have not
+    ** yet been computed. */
+    int ii, jj;
+    char *zP4 = pOp->p4.z;
+    assert( zP4!=0 );
+    assert( pOp->p4type==P4_DYNAMIC );
+    for(ii=jj=0; zP4[jj]; ii++){
+      if( pTab->aCol[ii].colFlags & COLFLAG_VIRTUAL ){
+        continue;
+      }
+      if( pTab->aCol[ii].colFlags & COLFLAG_STORED ){
+        zP4[jj] = SQLITE_AFF_NONE;
+      }
+      jj++;
+    }
+  }
+
+  /* Because there can be multiple generated columns that refer to one another,
+  ** this is a two-pass algorithm.  On the first pass, mark all generated
+  ** columns as "not available".
+  */
+  for(i=0; i<pTab->nCol; i++){
+    if( pTab->aCol[i].colFlags & COLFLAG_GENERATED ){
+      testcase( pTab->aCol[i].colFlags & COLFLAG_VIRTUAL );
+      testcase( pTab->aCol[i].colFlags & COLFLAG_STORED );
+      pTab->aCol[i].colFlags |= COLFLAG_NOTAVAIL;
+    }
+  }
+
+  w.u.pTab = pTab;
+  w.xExprCallback = exprColumnFlagUnion;
+  w.xSelectCallback = 0;
+  w.xSelectCallback2 = 0;
+
+  /* On the second pass, compute the value of each NOT-AVAILABLE column.
+  ** Companion code in the TK_COLUMN case of sqlite3ExprCodeTarget() will
+  ** compute dependencies and mark remove the COLSPAN_NOTAVAIL mark, as
+  ** they are needed.
+  */
+  pParse->iSelfTab = -iRegStore;
+  do{
+    eProgress = 0;
+    pRedo = 0;
+    for(i=0; i<pTab->nCol; i++){
+      Column *pCol = pTab->aCol + i;
+      if( (pCol->colFlags & COLFLAG_NOTAVAIL)!=0 ){
+        int x;
+        pCol->colFlags |= COLFLAG_BUSY;
+        w.eCode = 0;
+        sqlite3WalkExpr(&w, pCol->pDflt);
+        pCol->colFlags &= ~COLFLAG_BUSY;
+        if( w.eCode & COLFLAG_NOTAVAIL ){
+          pRedo = pCol;
+          continue;
+        }
+        eProgress = 1;
+        assert( pCol->colFlags & COLFLAG_GENERATED );
+        x = sqlite3TableColumnToStorage(pTab, i) + iRegStore;
+        sqlite3ExprCodeGeneratedColumn(pParse, pCol, x);
+        pCol->colFlags &= ~COLFLAG_NOTAVAIL;
+      }
+    }
+  }while( pRedo && eProgress );
+  if( pRedo ){
+    sqlite3ErrorMsg(pParse, "generated column loop on \"%s\"", pRedo->zName);
+  }
+  pParse->iSelfTab = 0;
+}
+#endif /* SQLITE_OMIT_GENERATED_COLUMNS */
+
+
 #ifndef SQLITE_OMIT_AUTOINCREMENT
 /*
 ** Locate or create an AutoincInfo structure associated with table pTab
@@ -116957,7 +119753,7 @@ SQLITE_PRIVATE void sqlite3Insert(
   Parse *pParse,        /* Parser context */
   SrcList *pTabList,    /* Name of table into which we are inserting */
   Select *pSelect,      /* A SELECT statement to use as the data source */
-  IdList *pColumn,      /* Column names corresponding to IDLIST. */
+  IdList *pColumn,      /* Column names corresponding to IDLIST, or NULL. */
   int onError,          /* How to handle constraint errors */
   Upsert *pUpsert       /* ON CONFLICT clauses for upsert, or NULL */
 ){
@@ -116982,6 +119778,7 @@ SQLITE_PRIVATE void sqlite3Insert(
   u8 withoutRowid;      /* 0 for normal table.  1 for WITHOUT ROWID table */
   u8 bIdListInOrder;    /* True if IDLIST is in table order */
   ExprList *pList = 0;  /* List of VALUES() to be inserted  */
+  int iRegStore;        /* Register in which to store next column */
 
   /* Register allocations */
   int regFromSelect = 0;/* Base register for data coming from SELECT */
@@ -117089,8 +119886,8 @@ SQLITE_PRIVATE void sqlite3Insert(
   */
   regAutoinc = autoIncBegin(pParse, iDb, pTab);
 
-  /* Allocate registers for holding the rowid of the new row,
-  ** the content of the new row, and the assembled row record.
+  /* Allocate a block registers to hold the rowid and the values
+  ** for all columns of the new row.
   */
   regRowid = regIns = pParse->nMem+1;
   pParse->nMem += pTab->nCol + 1;
@@ -117109,9 +119906,17 @@ SQLITE_PRIVATE void sqlite3Insert(
   ** the index into IDLIST of the primary key column.  ipkColumn is
   ** the index of the primary key as it appears in IDLIST, not as
   ** is appears in the original table.  (The index of the INTEGER
-  ** PRIMARY KEY in the original table is pTab->iPKey.)
+  ** PRIMARY KEY in the original table is pTab->iPKey.)  After this
+  ** loop, if ipkColumn==(-1), that means that integer primary key
+  ** is unspecified, and hence the table is either WITHOUT ROWID or
+  ** it will automatically generated an integer primary key.
+  **
+  ** bIdListInOrder is true if the columns in IDLIST are in storage
+  ** order.  This enables an optimization that avoids shuffling the
+  ** columns into storage order.  False negatives are harmless,
+  ** but false positives will cause database corruption.
   */
-  bIdListInOrder = (pTab->tabFlags & TF_OOOHidden)==0;
+  bIdListInOrder = (pTab->tabFlags & (TF_OOOHidden|TF_HasStored))==0;
   if( pColumn ){
     for(i=0; i<pColumn->nId; i++){
       pColumn->a[i].idx = -1;
@@ -117124,6 +119929,14 @@ SQLITE_PRIVATE void sqlite3Insert(
           if( j==pTab->iPKey ){
             ipkColumn = i;  assert( !withoutRowid );
           }
+#ifndef SQLITE_OMIT_GENERATED_COLUMNS
+          if( pTab->aCol[j].colFlags & (COLFLAG_STORED|COLFLAG_VIRTUAL) ){
+            sqlite3ErrorMsg(pParse, 
+               "cannot INSERT into generated column \"%s\"",
+               pTab->aCol[j].zName);
+            goto insert_cleanup;
+          }
+#endif
           break;
         }
       }
@@ -117233,13 +120046,26 @@ SQLITE_PRIVATE void sqlite3Insert(
   */
   if( pColumn==0 && nColumn>0 ){
     ipkColumn = pTab->iPKey;
+#ifndef SQLITE_OMIT_GENERATED_COLUMNS
+    if( ipkColumn>=0 && (pTab->tabFlags & TF_HasGenerated)!=0 ){
+      testcase( pTab->tabFlags & TF_HasVirtual );
+      testcase( pTab->tabFlags & TF_HasStored );
+      for(i=ipkColumn-1; i>=0; i--){
+        if( pTab->aCol[i].colFlags & COLFLAG_GENERATED ){
+          testcase( pTab->aCol[i].colFlags & COLFLAG_VIRTUAL );
+          testcase( pTab->aCol[i].colFlags & COLFLAG_STORED );
+          ipkColumn--;
+        }
+      }
+    }
+#endif
   }
 
   /* Make sure the number of columns in the source data matches the number
   ** of columns to be inserted into the table.
   */
   for(i=0; i<pTab->nCol; i++){
-    nHidden += (IsHiddenColumn(&pTab->aCol[i]) ? 1 : 0);
+    if( pTab->aCol[i].colFlags & COLFLAG_NOINSERT ) nHidden++;
   }
   if( pColumn==0 && nColumn && nColumn!=(pTab->nCol-nHidden) ){
     sqlite3ErrorMsg(pParse, 
@@ -117267,7 +120093,7 @@ SQLITE_PRIVATE void sqlite3Insert(
     int nIdx;
     nIdx = sqlite3OpenTableAndIndices(pParse, pTab, OP_OpenWrite, 0, -1, 0,
                                       &iDataCur, &iIdxCur);
-    aRegIdx = sqlite3DbMallocRawNN(db, sizeof(int)*(nIdx+1));
+    aRegIdx = sqlite3DbMallocRawNN(db, sizeof(int)*(nIdx+2));
     if( aRegIdx==0 ){
       goto insert_cleanup;
     }
@@ -117276,6 +120102,7 @@ SQLITE_PRIVATE void sqlite3Insert(
       aRegIdx[i] = ++pParse->nMem;
       pParse->nMem += pIdx->nColumn;
     }
+    aRegIdx[i] = ++pParse->nMem;  /* Register to store the table record */
   }
 #ifndef SQLITE_OMIT_UPSERT
   if( pUpsert ){
@@ -117284,6 +120111,13 @@ SQLITE_PRIVATE void sqlite3Insert(
               pTab->zName);
       goto insert_cleanup;
     }
+    if( pTab->pSelect ){
+      sqlite3ErrorMsg(pParse, "cannot UPSERT a view");
+      goto insert_cleanup;
+    }
+    if( sqlite3HasExplicitNulls(pParse, pUpsert->pUpsertTarget) ){
+      goto insert_cleanup;
+    }
     pTabList->a[0].iCursor = iDataCur;
     pUpsert->pUpsertSrc = pTabList;
     pUpsert->regData = regData;
@@ -117318,10 +120152,91 @@ SQLITE_PRIVATE void sqlite3Insert(
     **         goto C
     **      D: ...
     */
+    sqlite3VdbeReleaseRegisters(pParse, regData, pTab->nCol, 0, 0);
     addrInsTop = addrCont = sqlite3VdbeAddOp1(v, OP_Yield, dest.iSDParm);
     VdbeCoverage(v);
+    if( ipkColumn>=0 ){
+      /* tag-20191021-001: If the INTEGER PRIMARY KEY is being generated by the
+      ** SELECT, go ahead and copy the value into the rowid slot now, so that
+      ** the value does not get overwritten by a NULL at tag-20191021-002. */
+      sqlite3VdbeAddOp2(v, OP_Copy, regFromSelect+ipkColumn, regRowid);
+    }
   }
 
+  /* Compute data for ordinary columns of the new entry.  Values
+  ** are written in storage order into registers starting with regData.
+  ** Only ordinary columns are computed in this loop. The rowid
+  ** (if there is one) is computed later and generated columns are
+  ** computed after the rowid since they might depend on the value
+  ** of the rowid.
+  */
+  nHidden = 0;
+  iRegStore = regData;  assert( regData==regRowid+1 );
+  for(i=0; i<pTab->nCol; i++, iRegStore++){
+    int k;
+    u32 colFlags;
+    assert( i>=nHidden );
+    if( i==pTab->iPKey ){
+      /* tag-20191021-002: References to the INTEGER PRIMARY KEY are filled
+      ** using the rowid. So put a NULL in the IPK slot of the record to avoid
+      ** using excess space.  The file format definition requires this extra
+      ** NULL - we cannot optimize further by skipping the column completely */
+      sqlite3VdbeAddOp1(v, OP_SoftNull, iRegStore);
+      continue;
+    }
+    if( ((colFlags = pTab->aCol[i].colFlags) & COLFLAG_NOINSERT)!=0 ){
+      nHidden++;
+      if( (colFlags & COLFLAG_VIRTUAL)!=0 ){
+        /* Virtual columns do not participate in OP_MakeRecord.  So back up
+        ** iRegStore by one slot to compensate for the iRegStore++ in the
+        ** outer for() loop */
+        iRegStore--;
+        continue;
+      }else if( (colFlags & COLFLAG_STORED)!=0 ){
+        /* Stored columns are computed later.  But if there are BEFORE
+        ** triggers, the slots used for stored columns will be OP_Copy-ed
+        ** to a second block of registers, so the register needs to be
+        ** initialized to NULL to avoid an uninitialized register read */
+        if( tmask & TRIGGER_BEFORE ){
+          sqlite3VdbeAddOp1(v, OP_SoftNull, iRegStore);
+        }
+        continue;
+      }else if( pColumn==0 ){
+        /* Hidden columns that are not explicitly named in the INSERT
+        ** get there default value */
+        sqlite3ExprCodeFactorable(pParse, pTab->aCol[i].pDflt, iRegStore);
+        continue;
+      }
+    }
+    if( pColumn ){
+      for(j=0; j<pColumn->nId && pColumn->a[j].idx!=i; j++){}
+      if( j>=pColumn->nId ){
+        /* A column not named in the insert column list gets its
+        ** default value */
+        sqlite3ExprCodeFactorable(pParse, pTab->aCol[i].pDflt, iRegStore);
+        continue;
+      }
+      k = j;
+    }else if( nColumn==0 ){
+      /* This is INSERT INTO ... DEFAULT VALUES.  Load the default value. */
+      sqlite3ExprCodeFactorable(pParse, pTab->aCol[i].pDflt, iRegStore);
+      continue;
+    }else{
+      k = i - nHidden;
+    }
+
+    if( useTempTable ){
+      sqlite3VdbeAddOp3(v, OP_Column, srcTab, k, iRegStore); 
+    }else if( pSelect ){
+      if( regFromSelect!=regData ){
+        sqlite3VdbeAddOp2(v, OP_SCopy, regFromSelect+k, iRegStore);
+      }
+    }else{
+      sqlite3ExprCode(pParse, pList->a[k].pExpr, iRegStore);
+    }
+  }
+
+
   /* Run the BEFORE and INSTEAD OF triggers, if there are any
   */
   endOfLoop = sqlite3VdbeMakeLabel(pParse);
@@ -117356,25 +120271,21 @@ SQLITE_PRIVATE void sqlite3Insert(
     */
     assert( !IsVirtual(pTab) );
 
-    /* Create the new column data
-    */
-    for(i=j=0; i<pTab->nCol; i++){
-      if( pColumn ){
-        for(j=0; j<pColumn->nId; j++){
-          if( pColumn->a[j].idx==i ) break;
-        }
-      }
-      if( (!useTempTable && !pList) || (pColumn && j>=pColumn->nId)
-            || (pColumn==0 && IsOrdinaryHiddenColumn(&pTab->aCol[i])) ){
-        sqlite3ExprCode(pParse, pTab->aCol[i].pDflt, regCols+i+1);
-      }else if( useTempTable ){
-        sqlite3VdbeAddOp3(v, OP_Column, srcTab, j, regCols+i+1); 
-      }else{
-        assert( pSelect==0 ); /* Otherwise useTempTable is true */
-        sqlite3ExprCodeAndCache(pParse, pList->a[j].pExpr, regCols+i+1);
-      }
-      if( pColumn==0 && !IsOrdinaryHiddenColumn(&pTab->aCol[i]) ) j++;
+    /* Copy the new data already generated. */
+    assert( pTab->nNVCol>0 );
+    sqlite3VdbeAddOp3(v, OP_Copy, regRowid+1, regCols+1, pTab->nNVCol-1);
+
+#ifndef SQLITE_OMIT_GENERATED_COLUMNS
+    /* Compute the new value for generated columns after all other
+    ** columns have already been computed.  This must be done after
+    ** computing the ROWID in case one of the generated columns
+    ** refers to the ROWID. */
+    if( pTab->tabFlags & TF_HasGenerated ){
+      testcase( pTab->tabFlags & TF_HasVirtual );
+      testcase( pTab->tabFlags & TF_HasStored );
+      sqlite3ComputeGeneratedColumns(pParse, regCols+1, pTab);
     }
+#endif
 
     /* If this is an INSERT on a view with an INSTEAD OF INSERT trigger,
     ** do not attempt any conversions before assembling the record.
@@ -117392,19 +120303,17 @@ SQLITE_PRIVATE void sqlite3Insert(
     sqlite3ReleaseTempRange(pParse, regCols, pTab->nCol+1);
   }
 
-  /* Compute the content of the next row to insert into a range of
-  ** registers beginning at regIns.
-  */
   if( !isView ){
     if( IsVirtual(pTab) ){
       /* The row that the VUpdate opcode will delete: none */
       sqlite3VdbeAddOp2(v, OP_Null, 0, regIns);
     }
     if( ipkColumn>=0 ){
+      /* Compute the new rowid */
       if( useTempTable ){
         sqlite3VdbeAddOp3(v, OP_Column, srcTab, ipkColumn, regRowid);
       }else if( pSelect ){
-        sqlite3VdbeAddOp2(v, OP_Copy, regFromSelect+ipkColumn, regRowid);
+        /* Rowid already initialized at tag-20191021-001 */
       }else{
         Expr *pIpk = pList->a[ipkColumn].pExpr;
         if( pIpk->op==TK_NULL && !IsVirtual(pTab) ){
@@ -117437,45 +120346,15 @@ SQLITE_PRIVATE void sqlite3Insert(
     }
     autoIncStep(pParse, regAutoinc, regRowid);
 
-    /* Compute data for all columns of the new entry, beginning
-    ** with the first column.
-    */
-    nHidden = 0;
-    for(i=0; i<pTab->nCol; i++){
-      int iRegStore = regRowid+1+i;
-      if( i==pTab->iPKey ){
-        /* The value of the INTEGER PRIMARY KEY column is always a NULL.
-        ** Whenever this column is read, the rowid will be substituted
-        ** in its place.  Hence, fill this column with a NULL to avoid
-        ** taking up data space with information that will never be used.
-        ** As there may be shallow copies of this value, make it a soft-NULL */
-        sqlite3VdbeAddOp1(v, OP_SoftNull, iRegStore);
-        continue;
-      }
-      if( pColumn==0 ){
-        if( IsHiddenColumn(&pTab->aCol[i]) ){
-          j = -1;
-          nHidden++;
-        }else{
-          j = i - nHidden;
-        }
-      }else{
-        for(j=0; j<pColumn->nId; j++){
-          if( pColumn->a[j].idx==i ) break;
-        }
-      }
-      if( j<0 || nColumn==0 || (pColumn && j>=pColumn->nId) ){
-        sqlite3ExprCodeFactorable(pParse, pTab->aCol[i].pDflt, iRegStore);
-      }else if( useTempTable ){
-        sqlite3VdbeAddOp3(v, OP_Column, srcTab, j, iRegStore); 
-      }else if( pSelect ){
-        if( regFromSelect!=regData ){
-          sqlite3VdbeAddOp2(v, OP_SCopy, regFromSelect+j, iRegStore);
-        }
-      }else{
-        sqlite3ExprCode(pParse, pList->a[j].pExpr, iRegStore);
-      }
+#ifndef SQLITE_OMIT_GENERATED_COLUMNS
+    /* Compute the new value for generated columns after all other
+    ** columns have already been computed.  This must be done after
+    ** computing the ROWID in case one of the generated columns
+    ** is derived from the INTEGER PRIMARY KEY. */
+    if( pTab->tabFlags & TF_HasGenerated ){
+      sqlite3ComputeGeneratedColumns(pParse, regRowid+1, pTab);
     }
+#endif
 
     /* Generate code to check constraints and generate index keys and
     ** do the insertion.
@@ -117505,9 +120384,7 @@ SQLITE_PRIVATE void sqlite3Insert(
       ** cursor that is disturbed. And these instructions both clear the
       ** VdbeCursor.seekResult variable, disabling the OPFLAG_USESEEKRESULT
       ** functionality.  */
-      bUseSeek = (isReplace==0 || (pTrigger==0 &&
-          ((db->flags & SQLITE_ForeignKeys)==0 || sqlite3FkReferences(pTab)==0)
-      ));
+      bUseSeek = (isReplace==0 || !sqlite3VdbeHasSubProgram(v));
       sqlite3CompleteInsertion(pParse, pTab, iDataCur, iIdxCur,
           regIns, aRegIdx, 0, appendFlag, bUseSeek
       );
@@ -117536,6 +120413,15 @@ SQLITE_PRIVATE void sqlite3Insert(
     sqlite3VdbeAddOp1(v, OP_Close, srcTab);
   }else if( pSelect ){
     sqlite3VdbeGoto(v, addrCont);
+#ifdef SQLITE_DEBUG
+    /* If we are jumping back to an OP_Yield that is preceded by an
+    ** OP_ReleaseReg, set the p5 flag on the OP_Goto so that the
+    ** OP_ReleaseReg will be included in the loop. */
+    if( sqlite3VdbeGetOp(v, addrCont-1)->opcode==OP_ReleaseReg ){
+      assert( sqlite3VdbeGetOp(v, addrCont)->opcode==OP_Yield );
+      sqlite3VdbeChangeP5(v, 1);
+    }
+#endif
     sqlite3VdbeJumpHere(v, addrInsTop);
   }
 
@@ -117679,6 +120565,14 @@ SQLITE_PRIVATE int sqlite3ExprReferencesUpdatedColumn(
 ** the same as the order of indices on the linked list of indices
 ** at pTab->pIndex.
 **
+** (2019-05-07) The generated code also creates a new record for the
+** main table, if pTab is a rowid table, and stores that record in the
+** register identified by aRegIdx[nIdx] - in other words in the first
+** entry of aRegIdx[] past the last index.  It is important that the
+** record be generated during constraint checks to avoid affinity changes
+** to the register content that occur after constraint checks but before
+** the new record is inserted.
+**
 ** The caller must have already opened writeable cursors on the main
 ** table and all applicable indices (that is to say, all indices for which
 ** aRegIdx[] is not zero).  iDataCur is the cursor for the main table when
@@ -117750,7 +120644,6 @@ SQLITE_PRIVATE void sqlite3GenerateConstraintChecks(
   int ix;              /* Index loop counter */
   int nCol;            /* Number of columns */
   int onError;         /* Conflict resolution strategy */
-  int addr1;           /* Address of jump instruction */
   int seenReplace = 0; /* True if REPLACE is used to resolve INT PK conflict */
   int nPkField;        /* Number of fields in PRIMARY KEY. 1 for ROWID tables */
   Index *pUpIdx = 0;   /* Index to which to apply the upsert */
@@ -117760,6 +120653,13 @@ SQLITE_PRIVATE void sqlite3GenerateConstraintChecks(
   int upsertJump = 0;    /* Address of Goto that jumps into upsert subroutine */
   int ipkTop = 0;        /* Top of the IPK uniqueness check */
   int ipkBottom = 0;     /* OP_Goto at the end of the IPK uniqueness check */
+  /* Variables associated with retesting uniqueness constraints after
+  ** replace triggers fire have run */
+  int regTrigCnt;       /* Register used to count replace trigger invocations */
+  int addrRecheck = 0;  /* Jump here to recheck all uniqueness constraints */
+  int lblRecheckOk = 0; /* Each recheck jumps to this label if it passes */
+  Trigger *pTrigger;    /* List of DELETE triggers on the table pTab */
+  int nReplaceTrig = 0; /* Number of replace triggers coded */
 
   isUpdate = regOldData!=0;
   db = pParse->db;
@@ -117786,63 +120686,103 @@ SQLITE_PRIVATE void sqlite3GenerateConstraintChecks(
 
   /* Test all NOT NULL constraints.
   */
-  for(i=0; i<nCol; i++){
-    if( i==pTab->iPKey ){
-      continue;        /* ROWID is never NULL */
-    }
-    if( aiChng && aiChng[i]<0 ){
-      /* Don't bother checking for NOT NULL on columns that do not change */
-      continue;
-    }
-    onError = pTab->aCol[i].notNull;
-    if( onError==OE_None ) continue;  /* This column is allowed to be NULL */
-    if( overrideError!=OE_Default ){
-      onError = overrideError;
-    }else if( onError==OE_Default ){
-      onError = OE_Abort;
-    }
-    if( onError==OE_Replace && pTab->aCol[i].pDflt==0 ){
-      onError = OE_Abort;
-    }
-    assert( onError==OE_Rollback || onError==OE_Abort || onError==OE_Fail
-        || onError==OE_Ignore || onError==OE_Replace );
-    addr1 = 0;
-    switch( onError ){
-      case OE_Replace: {
-        assert( onError==OE_Replace );
-        addr1 = sqlite3VdbeMakeLabel(pParse);
-        sqlite3VdbeAddOp2(v, OP_NotNull, regNewData+1+i, addr1);
-          VdbeCoverage(v);
-        sqlite3ExprCode(pParse, pTab->aCol[i].pDflt, regNewData+1+i);
-        sqlite3VdbeAddOp2(v, OP_NotNull, regNewData+1+i, addr1);
-          VdbeCoverage(v);
-        onError = OE_Abort;
-        /* Fall through into the OE_Abort case to generate code that runs
-        ** if both the input and the default value are NULL */
-      }
-      case OE_Abort:
-        sqlite3MayAbort(pParse);
-        /* Fall through */
-      case OE_Rollback:
-      case OE_Fail: {
-        char *zMsg = sqlite3MPrintf(db, "%s.%s", pTab->zName,
-                                    pTab->aCol[i].zName);
-        sqlite3VdbeAddOp3(v, OP_HaltIfNull, SQLITE_CONSTRAINT_NOTNULL, onError,
-                          regNewData+1+i);
-        sqlite3VdbeAppendP4(v, zMsg, P4_DYNAMIC);
-        sqlite3VdbeChangeP5(v, P5_ConstraintNotNull);
-        VdbeCoverage(v);
-        if( addr1 ) sqlite3VdbeResolveLabel(v, addr1);
+  if( pTab->tabFlags & TF_HasNotNull ){
+    int b2ndPass = 0;         /* True if currently running 2nd pass */
+    int nSeenReplace = 0;     /* Number of ON CONFLICT REPLACE operations */
+    int nGenerated = 0;       /* Number of generated columns with NOT NULL */
+    while(1){  /* Make 2 passes over columns. Exit loop via "break" */
+      for(i=0; i<nCol; i++){
+        int iReg;                        /* Register holding column value */
+        Column *pCol = &pTab->aCol[i];   /* The column to check for NOT NULL */
+        int isGenerated;                 /* non-zero if column is generated */
+        onError = pCol->notNull;
+        if( onError==OE_None ) continue; /* No NOT NULL on this column */
+        if( i==pTab->iPKey ){
+          continue;        /* ROWID is never NULL */
+        }
+        isGenerated = pCol->colFlags & COLFLAG_GENERATED;
+        if( isGenerated && !b2ndPass ){
+          nGenerated++;
+          continue;        /* Generated columns processed on 2nd pass */
+        }
+        if( aiChng && aiChng[i]<0 && !isGenerated ){
+          /* Do not check NOT NULL on columns that do not change */
+          continue;
+        }
+        if( overrideError!=OE_Default ){
+          onError = overrideError;
+        }else if( onError==OE_Default ){
+          onError = OE_Abort;
+        }
+        if( onError==OE_Replace ){
+          if( b2ndPass        /* REPLACE becomes ABORT on the 2nd pass */
+           || pCol->pDflt==0  /* REPLACE is ABORT if no DEFAULT value */
+          ){
+            testcase( pCol->colFlags & COLFLAG_VIRTUAL );
+            testcase( pCol->colFlags & COLFLAG_STORED );
+            testcase( pCol->colFlags & COLFLAG_GENERATED );
+            onError = OE_Abort;
+          }else{
+            assert( !isGenerated );
+          }
+        }else if( b2ndPass && !isGenerated ){
+          continue;
+        }
+        assert( onError==OE_Rollback || onError==OE_Abort || onError==OE_Fail
+            || onError==OE_Ignore || onError==OE_Replace );
+        testcase( i!=sqlite3TableColumnToStorage(pTab, i) );
+        iReg = sqlite3TableColumnToStorage(pTab, i) + regNewData + 1;
+        switch( onError ){
+          case OE_Replace: {
+            int addr1 = sqlite3VdbeAddOp1(v, OP_NotNull, iReg);
+            VdbeCoverage(v);
+            assert( (pCol->colFlags & COLFLAG_GENERATED)==0 );
+            nSeenReplace++;
+            sqlite3ExprCode(pParse, pCol->pDflt, iReg);
+            sqlite3VdbeJumpHere(v, addr1);
+            break;
+          }
+          case OE_Abort:
+            sqlite3MayAbort(pParse);
+            /* Fall through */
+          case OE_Rollback:
+          case OE_Fail: {
+            char *zMsg = sqlite3MPrintf(db, "%s.%s", pTab->zName,
+                                        pCol->zName);
+            sqlite3VdbeAddOp3(v, OP_HaltIfNull, SQLITE_CONSTRAINT_NOTNULL,
+                              onError, iReg);
+            sqlite3VdbeAppendP4(v, zMsg, P4_DYNAMIC);
+            sqlite3VdbeChangeP5(v, P5_ConstraintNotNull);
+            VdbeCoverage(v);
+            break;
+          }
+          default: {
+            assert( onError==OE_Ignore );
+            sqlite3VdbeAddOp2(v, OP_IsNull, iReg, ignoreDest);
+            VdbeCoverage(v);
+            break;
+          }
+        } /* end switch(onError) */
+      } /* end loop i over columns */
+      if( nGenerated==0 && nSeenReplace==0 ){
+        /* If there are no generated columns with NOT NULL constraints
+        ** and no NOT NULL ON CONFLICT REPLACE constraints, then a single
+        ** pass is sufficient */
         break;
       }
-      default: {
-        assert( onError==OE_Ignore );
-        sqlite3VdbeAddOp2(v, OP_IsNull, regNewData+1+i, ignoreDest);
-        VdbeCoverage(v);
-        break;
+      if( b2ndPass ) break;  /* Never need more than 2 passes */
+      b2ndPass = 1;
+#ifndef SQLITE_OMIT_GENERATED_COLUMNS
+      if( nSeenReplace>0 && (pTab->tabFlags & TF_HasGenerated)!=0 ){
+        /* If any NOT NULL ON CONFLICT REPLACE constraints fired on the
+        ** first pass, recomputed values for all generated columns, as
+        ** those values might depend on columns affected by the REPLACE.
+        */
+        sqlite3ComputeGeneratedColumns(pParse, regNewData+1, pTab);
       }
-    }
-  }
+#endif
+    } /* end of 2-pass loop */
+  } /* end if( has-not-null-constraints ) */
 
   /* Test all CHECK constraints
   */
@@ -117867,9 +120807,9 @@ SQLITE_PRIVATE void sqlite3GenerateConstraintChecks(
       if( onError==OE_Ignore ){
         sqlite3VdbeGoto(v, ignoreDest);
       }else{
-        char *zName = pCheck->a[i].zName;
+        char *zName = pCheck->a[i].zEName;
         if( zName==0 ) zName = pTab->zName;
-        if( onError==OE_Replace ) onError = OE_Abort; /* IMP: R-15569-63625 */
+        if( onError==OE_Replace ) onError = OE_Abort; /* IMP: R-26383-51744 */
         sqlite3HaltConstraint(pParse, SQLITE_CONSTRAINT_CHECK,
                               onError, zName, P4_TRANSIENT,
                               P5_ConstraintCheck);
@@ -117924,6 +120864,50 @@ SQLITE_PRIVATE void sqlite3GenerateConstraintChecks(
     }
   }
 
+  /* Determine if it is possible that triggers (either explicitly coded
+  ** triggers or FK resolution actions) might run as a result of deletes
+  ** that happen when OE_Replace conflict resolution occurs. (Call these
+  ** "replace triggers".)  If any replace triggers run, we will need to
+  ** recheck all of the uniqueness constraints after they have all run.
+  ** But on the recheck, the resolution is OE_Abort instead of OE_Replace.
+  **
+  ** If replace triggers are a possibility, then
+  **
+  **   (1) Allocate register regTrigCnt and initialize it to zero.
+  **       That register will count the number of replace triggers that
+  **       fire.  Constraint recheck only occurs if the number is positive.
+  **   (2) Initialize pTrigger to the list of all DELETE triggers on pTab.
+  **   (3) Initialize addrRecheck and lblRecheckOk
+  **
+  ** The uniqueness rechecking code will create a series of tests to run
+  ** in a second pass.  The addrRecheck and lblRecheckOk variables are
+  ** used to link together these tests which are separated from each other
+  ** in the generate bytecode.
+  */
+  if( (db->flags & (SQLITE_RecTriggers|SQLITE_ForeignKeys))==0 ){
+    /* There are not DELETE triggers nor FK constraints.  No constraint
+    ** rechecks are needed. */
+    pTrigger = 0;
+    regTrigCnt = 0;
+  }else{
+    if( db->flags&SQLITE_RecTriggers ){
+      pTrigger = sqlite3TriggersExist(pParse, pTab, TK_DELETE, 0, 0);
+      regTrigCnt = pTrigger!=0 || sqlite3FkRequired(pParse, pTab, 0, 0);
+    }else{
+      pTrigger = 0;
+      regTrigCnt = sqlite3FkRequired(pParse, pTab, 0, 0);
+    }
+    if( regTrigCnt ){
+      /* Replace triggers might exist.  Allocate the counter and
+      ** initialize it to zero. */
+      regTrigCnt = ++pParse->nMem;
+      sqlite3VdbeAddOp2(v, OP_Integer, 0, regTrigCnt);
+      VdbeComment((v, "trigger count"));
+      lblRecheckOk = sqlite3VdbeMakeLabel(pParse);
+      addrRecheck = lblRecheckOk;
+    }
+  }
+
   /* If rowid is changing, make sure the new rowid does not previously
   ** exist in the table.
   */
@@ -118013,14 +120997,12 @@ SQLITE_PRIVATE void sqlite3GenerateConstraintChecks(
         ** to run without a statement journal if there are no indexes on the
         ** table.
         */
-        Trigger *pTrigger = 0;
-        if( db->flags&SQLITE_RecTriggers ){
-          pTrigger = sqlite3TriggersExist(pParse, pTab, TK_DELETE, 0, 0);
-        }
-        if( pTrigger || sqlite3FkRequired(pParse, pTab, 0, 0) ){
+        if( regTrigCnt ){
           sqlite3MultiWrite(pParse);
           sqlite3GenerateRowDelete(pParse, pTab, pTrigger, iDataCur, iIdxCur,
                                    regNewData, 1, 0, OE_Replace, 1, -1);
+          sqlite3VdbeAddOp2(v, OP_AddImm, regTrigCnt, 1); /* incr trigger cnt */
+          nReplaceTrig++;
         }else{
 #ifdef SQLITE_ENABLE_PREUPDATE_HOOK
           assert( HasRowid(pTab) );
@@ -118070,6 +121052,7 @@ SQLITE_PRIVATE void sqlite3GenerateConstraintChecks(
     int regR;            /* Range of registers holding conflicting PK */
     int iThisCur;        /* Cursor for this UNIQUE index */
     int addrUniqueOk;    /* Jump here if the UNIQUE constraint is satisfied */
+    int addrConflictCk;  /* First opcode in the conflict check logic */
 
     if( aRegIdx[ix]==0 ) continue;  /* Skip indices that do not change */
     if( pUpIdx==pIdx ){
@@ -118109,14 +121092,15 @@ SQLITE_PRIVATE void sqlite3GenerateConstraintChecks(
         sqlite3ExprCodeCopy(pParse, pIdx->aColExpr->a[i].pExpr, regIdx+i);
         pParse->iSelfTab = 0;
         VdbeComment((v, "%s column %d", pIdx->zName, i));
+      }else if( iField==XN_ROWID || iField==pTab->iPKey ){
+        x = regNewData;
+        sqlite3VdbeAddOp2(v, OP_IntCopy, x, regIdx+i);
+        VdbeComment((v, "rowid"));
       }else{
-        if( iField==XN_ROWID || iField==pTab->iPKey ){
-          x = regNewData;
-        }else{
-          x = iField + regNewData + 1;
-        }
-        sqlite3VdbeAddOp2(v, iField<0 ? OP_IntCopy : OP_SCopy, x, regIdx+i);
-        VdbeComment((v, "%s", iField<0 ? "rowid" : pTab->aCol[iField].zName));
+        testcase( sqlite3TableColumnToStorage(pTab, iField)!=iField );
+        x = sqlite3TableColumnToStorage(pTab, iField) + regNewData + 1;
+        sqlite3VdbeAddOp2(v, OP_SCopy, x, regIdx+i);
+        VdbeComment((v, "%s", pTab->aCol[iField].zName));
       }
     }
     sqlite3VdbeAddOp3(v, OP_MakeRecord, regIdx, pIdx->nColumn, aRegIdx[ix]);
@@ -118126,6 +121110,7 @@ SQLITE_PRIVATE void sqlite3GenerateConstraintChecks(
       sqlite3SetMakeRecordP5(v, pIdx->pTable);
     }
 #endif
+    sqlite3VdbeReleaseRegisters(pParse, regIdx, pIdx->nColumn, 0, 0);
 
     /* In an UPDATE operation, if this index is the PRIMARY KEY index 
     ** of a WITHOUT ROWID table and there has been no change the
@@ -118183,8 +121168,9 @@ SQLITE_PRIVATE void sqlite3GenerateConstraintChecks(
 
     /* Check to see if the new index entry will be unique */
     sqlite3VdbeVerifyAbortable(v, onError);
-    sqlite3VdbeAddOp4Int(v, OP_NoConflict, iThisCur, addrUniqueOk,
-                         regIdx, pIdx->nKeyCol); VdbeCoverage(v);
+    addrConflictCk = 
+      sqlite3VdbeAddOp4Int(v, OP_NoConflict, iThisCur, addrUniqueOk,
+                           regIdx, pIdx->nKeyCol); VdbeCoverage(v);
 
     /* Generate code to handle collisions */
     regR = (pIdx==pPk) ? regIdx : sqlite3GetTempRange(pParse, nPkField);
@@ -118205,7 +121191,7 @@ SQLITE_PRIVATE void sqlite3GenerateConstraintChecks(
         if( pIdx!=pPk ){
           for(i=0; i<pPk->nKeyCol; i++){
             assert( pPk->aiColumn[i]>=0 );
-            x = sqlite3ColumnOfIndex(pIdx, pPk->aiColumn[i]);
+            x = sqlite3TableColumnToIndex(pIdx, pPk->aiColumn[i]);
             sqlite3VdbeAddOp3(v, OP_Column, iThisCur, x, regR+i);
             VdbeComment((v, "%s.%s", pTab->zName,
                          pTab->aCol[pPk->aiColumn[i]].zName));
@@ -118231,6 +121217,7 @@ SQLITE_PRIVATE void sqlite3GenerateConstraintChecks(
               addrJump = addrUniqueOk;
               op = OP_Eq;
             }
+            x = sqlite3TableColumnToStorage(pTab, x);
             sqlite3VdbeAddOp4(v, op, 
                 regOldData+1+x, addrJump, regCmp+i, p4, P4_COLLSEQ
             );
@@ -118267,17 +121254,71 @@ SQLITE_PRIVATE void sqlite3GenerateConstraintChecks(
         break;
       }
       default: {
-        Trigger *pTrigger = 0;
+        int nConflictCk;   /* Number of opcodes in conflict check logic */
+
         assert( onError==OE_Replace );
-        if( db->flags&SQLITE_RecTriggers ){
-          pTrigger = sqlite3TriggersExist(pParse, pTab, TK_DELETE, 0, 0);
-        }
-        if( pTrigger || sqlite3FkRequired(pParse, pTab, 0, 0) ){
+        nConflictCk = sqlite3VdbeCurrentAddr(v) - addrConflictCk;
+        assert( nConflictCk>0 );
+        testcase( nConflictCk>1 );
+        if( regTrigCnt ){
           sqlite3MultiWrite(pParse);
+          nReplaceTrig++;
+        }
+        if( pTrigger && isUpdate ){
+          sqlite3VdbeAddOp1(v, OP_CursorLock, iDataCur);
         }
         sqlite3GenerateRowDelete(pParse, pTab, pTrigger, iDataCur, iIdxCur,
             regR, nPkField, 0, OE_Replace,
             (pIdx==pPk ? ONEPASS_SINGLE : ONEPASS_OFF), iThisCur);
+        if( pTrigger && isUpdate ){
+          sqlite3VdbeAddOp1(v, OP_CursorUnlock, iDataCur);
+        }
+        if( regTrigCnt ){
+          int addrBypass;  /* Jump destination to bypass recheck logic */
+
+          sqlite3VdbeAddOp2(v, OP_AddImm, regTrigCnt, 1); /* incr trigger cnt */
+          addrBypass = sqlite3VdbeAddOp0(v, OP_Goto);  /* Bypass recheck */
+          VdbeComment((v, "bypass recheck"));
+
+          /* Here we insert code that will be invoked after all constraint
+          ** checks have run, if and only if one or more replace triggers
+          ** fired. */
+          sqlite3VdbeResolveLabel(v, lblRecheckOk);
+          lblRecheckOk = sqlite3VdbeMakeLabel(pParse);
+          if( pIdx->pPartIdxWhere ){
+            /* Bypass the recheck if this partial index is not defined
+            ** for the current row */
+            sqlite3VdbeAddOp2(v, OP_IsNull, regIdx-1, lblRecheckOk);
+            VdbeCoverage(v);
+          }
+          /* Copy the constraint check code from above, except change
+          ** the constraint-ok jump destination to be the address of
+          ** the next retest block */
+          while( nConflictCk>0 ){
+            VdbeOp x;    /* Conflict check opcode to copy */
+            /* The sqlite3VdbeAddOp4() call might reallocate the opcode array.
+            ** Hence, make a complete copy of the opcode, rather than using
+            ** a pointer to the opcode. */
+            x = *sqlite3VdbeGetOp(v, addrConflictCk);
+            if( x.opcode!=OP_IdxRowid ){
+              int p2;      /* New P2 value for copied conflict check opcode */
+              if( sqlite3OpcodeProperty[x.opcode]&OPFLG_JUMP ){
+                p2 = lblRecheckOk;
+              }else{
+                p2 = x.p2;
+              }
+              sqlite3VdbeAddOp4(v, x.opcode, x.p1, p2, x.p3, x.p4.z, x.p4type);
+              sqlite3VdbeChangeP5(v, x.p5);
+              VdbeCoverageIf(v, p2!=x.p2);
+            }
+            nConflictCk--;
+            addrConflictCk++;
+          }
+          /* If the retest fails, issue an abort */
+          sqlite3UniqueConstraint(pParse, OE_Abort, pIdx);
+
+          sqlite3VdbeJumpHere(v, addrBypass); /* Terminate the recheck bypass */
+        }
         seenReplace = 1;
         break;
       }
@@ -118298,6 +121339,36 @@ SQLITE_PRIVATE void sqlite3GenerateConstraintChecks(
     sqlite3VdbeJumpHere(v, ipkBottom);
   }
 
+  /* Recheck all uniqueness constraints after replace triggers have run */
+  testcase( regTrigCnt!=0 && nReplaceTrig==0 );
+  assert( regTrigCnt!=0 || nReplaceTrig==0 );
+  if( nReplaceTrig ){
+    sqlite3VdbeAddOp2(v, OP_IfNot, regTrigCnt, lblRecheckOk);VdbeCoverage(v);
+    if( !pPk ){
+      if( isUpdate ){
+        sqlite3VdbeAddOp3(v, OP_Eq, regNewData, addrRecheck, regOldData);
+        sqlite3VdbeChangeP5(v, SQLITE_NOTNULL);
+        VdbeCoverage(v);
+      }
+      sqlite3VdbeAddOp3(v, OP_NotExists, iDataCur, addrRecheck, regNewData);
+      VdbeCoverage(v);
+      sqlite3RowidConstraint(pParse, OE_Abort, pTab);
+    }else{
+      sqlite3VdbeGoto(v, addrRecheck);
+    }
+    sqlite3VdbeResolveLabel(v, lblRecheckOk);
+  }
+
+  /* Generate the table record */
+  if( HasRowid(pTab) ){
+    int regRec = aRegIdx[ix];
+    sqlite3VdbeAddOp3(v, OP_MakeRecord, regNewData+1, pTab->nNVCol, regRec);
+    sqlite3SetMakeRecordP5(v, pTab);
+    if( !bAffinityDone ){
+      sqlite3TableAffinity(v, pTab, 0);
+    }
+  }
+
   *pbMayReplace = seenReplace;
   VdbeModuleComment((v, "END: GenCnstCks(%d)", seenReplace));
 }
@@ -118347,10 +121418,7 @@ SQLITE_PRIVATE void sqlite3CompleteInsertion(
   Vdbe *v;            /* Prepared statements under construction */
   Index *pIdx;        /* An index being inserted or updated */
   u8 pik_flags;       /* flag values passed to the btree insert */
-  int regData;        /* Content registers (after the rowid) */
-  int regRec;         /* Register holding assembled record for the table */
   int i;              /* Loop counter */
-  u8 bAffinityDone = 0; /* True if OP_Affinity has been run already */
 
   assert( update_flags==0
        || update_flags==OPFLAG_ISUPDATE
@@ -118361,8 +121429,11 @@ SQLITE_PRIVATE void sqlite3CompleteInsertion(
   assert( v!=0 );
   assert( pTab->pSelect==0 );  /* This table is not a VIEW */
   for(i=0, pIdx=pTab->pIndex; pIdx; pIdx=pIdx->pNext, i++){
+    /* All REPLACE indexes are at the end of the list */
+    assert( pIdx->onError!=OE_Replace
+         || pIdx->pNext==0
+         || pIdx->pNext->onError==OE_Replace );
     if( aRegIdx[i]==0 ) continue;
-    bAffinityDone = 1;
     if( pIdx->pPartIdxWhere ){
       sqlite3VdbeAddOp2(v, OP_IsNull, aRegIdx[i], sqlite3VdbeCurrentAddr(v)+2);
       VdbeCoverage(v);
@@ -118390,13 +121461,6 @@ SQLITE_PRIVATE void sqlite3CompleteInsertion(
     sqlite3VdbeChangeP5(v, pik_flags);
   }
   if( !HasRowid(pTab) ) return;
-  regData = regNewData + 1;
-  regRec = sqlite3GetTempReg(pParse);
-  sqlite3VdbeAddOp3(v, OP_MakeRecord, regData, pTab->nCol, regRec);
-  sqlite3SetMakeRecordP5(v, pTab);
-  if( !bAffinityDone ){
-    sqlite3TableAffinity(v, pTab, 0);
-  }
   if( pParse->nested ){
     pik_flags = 0;
   }else{
@@ -118409,7 +121473,7 @@ SQLITE_PRIVATE void sqlite3CompleteInsertion(
   if( useSeekResult ){
     pik_flags |= OPFLAG_USESEEKRESULT;
   }
-  sqlite3VdbeAddOp3(v, OP_Insert, iDataCur, regRec, regNewData);
+  sqlite3VdbeAddOp3(v, OP_Insert, iDataCur, aRegIdx[i], regNewData);
   if( !pParse->nested ){
     sqlite3VdbeAppendP4(v, pTab, P4_TABLE);
   }
@@ -118519,7 +121583,7 @@ static int xferCompatibleIndex(Index *pDest, Index *pSrc){
   int i;
   assert( pDest && pSrc );
   assert( pDest->pTable!=pSrc->pTable );
-  if( pDest->nKeyCol!=pSrc->nKeyCol ){
+  if( pDest->nKeyCol!=pSrc->nKeyCol || pDest->nColumn!=pSrc->nColumn ){
     return 0;   /* Different number of columns */
   }
   if( pDest->onError!=pSrc->onError ){
@@ -118696,6 +121760,39 @@ static int xferOptimization(
     ){
       return 0;    /* Neither table may have __hidden__ columns */
     }
+#endif
+#ifndef SQLITE_OMIT_GENERATED_COLUMNS
+    /* Even if tables t1 and t2 have identical schemas, if they contain
+    ** generated columns, then this statement is semantically incorrect:
+    **
+    **     INSERT INTO t2 SELECT * FROM t1;
+    **
+    ** The reason is that generated column values are returned by the
+    ** the SELECT statement on the right but the INSERT statement on the
+    ** left wants them to be omitted.
+    **
+    ** Nevertheless, this is a useful notational shorthand to tell SQLite
+    ** to do a bulk transfer all of the content from t1 over to t2.
+    ** 
+    ** We could, in theory, disable this (except for internal use by the
+    ** VACUUM command where it is actually needed).  But why do that?  It
+    ** seems harmless enough, and provides a useful service.
+    */
+    if( (pDestCol->colFlags & COLFLAG_GENERATED) !=
+        (pSrcCol->colFlags & COLFLAG_GENERATED) ){
+      return 0;    /* Both columns have the same generated-column type */
+    }
+    /* But the transfer is only allowed if both the source and destination
+    ** tables have the exact same expressions for generated columns.
+    ** This requirement could be relaxed for VIRTUAL columns, I suppose.
+    */
+    if( (pDestCol->colFlags & COLFLAG_GENERATED)!=0 ){
+      if( sqlite3ExprCompare(0, pSrcCol->pDflt, pDestCol->pDflt, -1)!=0 ){
+        testcase( pDestCol->colFlags & COLFLAG_VIRTUAL );
+        testcase( pDestCol->colFlags & COLFLAG_STORED );
+        return 0;  /* Different generator expressions */
+      }
+    }
 #endif
     if( pDestCol->affinity!=pSrcCol->affinity ){
       return 0;    /* Affinity must be the same on all columns */
@@ -118707,7 +121804,7 @@ static int xferOptimization(
       return 0;    /* tab2 must be NOT NULL if tab1 is */
     }
     /* Default values for second and subsequent columns need to match. */
-    if( i>0 ){
+    if( (pDestCol->colFlags & COLFLAG_GENERATED)==0 && i>0 ){
       assert( pDestCol->pDflt==0 || pDestCol->pDflt->op==TK_SPAN );
       assert( pSrcCol->pDflt==0 || pSrcCol->pDflt->op==TK_SPAN );
       if( (pDestCol->pDflt==0)!=(pSrcCol->pDflt==0) 
@@ -119390,6 +122487,14 @@ struct sqlite3_api_routines {
   /* Version 3.28.0 and later */
   int (*stmt_isexplain)(sqlite3_stmt*);
   int (*value_frombind)(sqlite3_value*);
+  /* Version 3.30.0 and later */
+  int (*drop_modules)(sqlite3*,const char**);
+  /* Version 3.31.0 and later */
+  sqlite3_int64 (*hard_heap_limit64)(sqlite3_int64);
+  const char *(*uri_key)(const char*,int);
+  const char *(*filename_database)(const char*);
+  const char *(*filename_journal)(const char*);
+  const char *(*filename_wal)(const char*);
 };
 
 /*
@@ -119680,8 +122785,16 @@ typedef int (*sqlite3_loadext_entry)(
 /* Version 3.26.0 and later */
 #define sqlite3_normalized_sql         sqlite3_api->normalized_sql
 /* Version 3.28.0 and later */
-#define sqlite3_stmt_isexplain         sqlite3_api->isexplain
-#define sqlite3_value_frombind         sqlite3_api->frombind
+#define sqlite3_stmt_isexplain         sqlite3_api->stmt_isexplain
+#define sqlite3_value_frombind         sqlite3_api->value_frombind
+/* Version 3.30.0 and later */
+#define sqlite3_drop_modules           sqlite3_api->drop_modules
+/* Version 3.31.0 and later */
+#define sqlite3_hard_heap_limit64      sqlite3_api->hard_heap_limit64
+#define sqlite3_uri_key                sqlite3_api->uri_key
+#define sqlite3_filename_database      sqlite3_api->filename_database
+#define sqlite3_filename_journal       sqlite3_api->filename_journal
+#define sqlite3_filename_wal           sqlite3_api->filename_wal
 #endif /* !defined(SQLITE_CORE) && !defined(SQLITE_OMIT_LOAD_EXTENSION) */
 
 #if !defined(SQLITE_CORE) && !defined(SQLITE_OMIT_LOAD_EXTENSION)
@@ -120147,7 +123260,19 @@ static const sqlite3_api_routines sqlite3Apis = {
 #endif
   /* Version 3.28.0 and later */
   sqlite3_stmt_isexplain,
-  sqlite3_value_frombind
+  sqlite3_value_frombind,
+  /* Version 3.30.0 and later */
+#ifndef SQLITE_OMIT_VIRTUALTABLE
+  sqlite3_drop_modules,
+#else
+  0,
+#endif
+  /* Version 3.31.0 and later */
+  sqlite3_hard_heap_limit64,
+  sqlite3_uri_key,
+  sqlite3_filename_database,
+  sqlite3_filename_journal,
+  sqlite3_filename_wal,
 };
 
 /*
@@ -120570,32 +123695,32 @@ SQLITE_PRIVATE void sqlite3AutoLoadExtensions(sqlite3 *db){
 #define PragTyp_FOREIGN_KEY_CHECK             13
 #define PragTyp_FOREIGN_KEY_LIST              14
 #define PragTyp_FUNCTION_LIST                 15
-#define PragTyp_INCREMENTAL_VACUUM            16
-#define PragTyp_INDEX_INFO                    17
-#define PragTyp_INDEX_LIST                    18
-#define PragTyp_INTEGRITY_CHECK               19
-#define PragTyp_JOURNAL_MODE                  20
-#define PragTyp_JOURNAL_SIZE_LIMIT            21
-#define PragTyp_LOCK_PROXY_FILE               22
-#define PragTyp_LOCKING_MODE                  23
-#define PragTyp_PAGE_COUNT                    24
-#define PragTyp_MMAP_SIZE                     25
-#define PragTyp_MODULE_LIST                   26
-#define PragTyp_OPTIMIZE                      27
-#define PragTyp_PAGE_SIZE                     28
-#define PragTyp_PRAGMA_LIST                   29
-#define PragTyp_SECURE_DELETE                 30
-#define PragTyp_SHRINK_MEMORY                 31
-#define PragTyp_SOFT_HEAP_LIMIT               32
-#define PragTyp_SYNCHRONOUS                   33
-#define PragTyp_TABLE_INFO                    34
-#define PragTyp_TEMP_STORE                    35
-#define PragTyp_TEMP_STORE_DIRECTORY          36
-#define PragTyp_THREADS                       37
-#define PragTyp_WAL_AUTOCHECKPOINT            38
-#define PragTyp_WAL_CHECKPOINT                39
-#define PragTyp_ACTIVATE_EXTENSIONS           40
-#define PragTyp_HEXKEY                        41
+#define PragTyp_HARD_HEAP_LIMIT               16
+#define PragTyp_INCREMENTAL_VACUUM            17
+#define PragTyp_INDEX_INFO                    18
+#define PragTyp_INDEX_LIST                    19
+#define PragTyp_INTEGRITY_CHECK               20
+#define PragTyp_JOURNAL_MODE                  21
+#define PragTyp_JOURNAL_SIZE_LIMIT            22
+#define PragTyp_LOCK_PROXY_FILE               23
+#define PragTyp_LOCKING_MODE                  24
+#define PragTyp_PAGE_COUNT                    25
+#define PragTyp_MMAP_SIZE                     26
+#define PragTyp_MODULE_LIST                   27
+#define PragTyp_OPTIMIZE                      28
+#define PragTyp_PAGE_SIZE                     29
+#define PragTyp_PRAGMA_LIST                   30
+#define PragTyp_SECURE_DELETE                 31
+#define PragTyp_SHRINK_MEMORY                 32
+#define PragTyp_SOFT_HEAP_LIMIT               33
+#define PragTyp_SYNCHRONOUS                   34
+#define PragTyp_TABLE_INFO                    35
+#define PragTyp_TEMP_STORE                    36
+#define PragTyp_TEMP_STORE_DIRECTORY          37
+#define PragTyp_THREADS                       38
+#define PragTyp_WAL_AUTOCHECKPOINT            39
+#define PragTyp_WAL_CHECKPOINT                40
+#define PragTyp_ACTIVATE_EXTENSIONS           41
 #define PragTyp_KEY                           42
 #define PragTyp_LOCK_STATUS                   43
 #define PragTyp_STATS                         44
@@ -120637,35 +123762,39 @@ static const char *const pragCName[] = {
   /*  18 */ "desc",       
   /*  19 */ "coll",       
   /*  20 */ "key",        
-  /*  21 */ "tbl",         /* Used by: stats */
-  /*  22 */ "idx",        
-  /*  23 */ "wdth",       
-  /*  24 */ "hght",       
-  /*  25 */ "flgs",       
-  /*  26 */ "seq",         /* Used by: index_list */
-  /*  27 */ "name",       
-  /*  28 */ "unique",     
-  /*  29 */ "origin",     
-  /*  30 */ "partial",    
-  /*  31 */ "table",       /* Used by: foreign_key_check */
-  /*  32 */ "rowid",      
-  /*  33 */ "parent",     
-  /*  34 */ "fkid",       
+  /*  21 */ "name",        /* Used by: function_list */
+  /*  22 */ "builtin",    
+  /*  23 */ "type",       
+  /*  24 */ "enc",        
+  /*  25 */ "narg",       
+  /*  26 */ "flags",      
+  /*  27 */ "tbl",         /* Used by: stats */
+  /*  28 */ "idx",        
+  /*  29 */ "wdth",       
+  /*  30 */ "hght",       
+  /*  31 */ "flgs",       
+  /*  32 */ "seq",         /* Used by: index_list */
+  /*  33 */ "name",       
+  /*  34 */ "unique",     
+  /*  35 */ "origin",     
+  /*  36 */ "partial",    
+  /*  37 */ "table",       /* Used by: foreign_key_check */
+  /*  38 */ "rowid",      
+  /*  39 */ "parent",     
+  /*  40 */ "fkid",       
                            /* index_info reuses 15 */
-  /*  35 */ "seq",         /* Used by: database_list */
-  /*  36 */ "name",       
-  /*  37 */ "file",       
-  /*  38 */ "busy",        /* Used by: wal_checkpoint */
-  /*  39 */ "log",        
-  /*  40 */ "checkpointed",
-  /*  41 */ "name",        /* Used by: function_list */
-  /*  42 */ "builtin",    
-                           /* collation_list reuses 26 */
-  /*  43 */ "database",    /* Used by: lock_status */
-  /*  44 */ "status",     
-  /*  45 */ "cache_size",  /* Used by: default_cache_size */
+  /*  41 */ "seq",         /* Used by: database_list */
+  /*  42 */ "name",       
+  /*  43 */ "file",       
+  /*  44 */ "busy",        /* Used by: wal_checkpoint */
+  /*  45 */ "log",        
+  /*  46 */ "checkpointed",
+                           /* collation_list reuses 32 */
+  /*  47 */ "database",    /* Used by: lock_status */
+  /*  48 */ "status",     
+  /*  49 */ "cache_size",  /* Used by: default_cache_size */
                            /* module_list pragma_list reuses 9 */
-  /*  46 */ "timeout",     /* Used by: busy_timeout */
+  /*  50 */ "timeout",     /* Used by: busy_timeout */
 };
 
 /* Definitions of all built-in pragmas */
@@ -120711,7 +123840,7 @@ static const PragmaName aPragmaName[] = {
  {/* zName:     */ "busy_timeout",
   /* ePragTyp:  */ PragTyp_BUSY_TIMEOUT,
   /* ePragFlg:  */ PragFlg_Result0,
-  /* ColNames:  */ 46, 1,
+  /* ColNames:  */ 50, 1,
   /* iArg:      */ 0 },
 #if !defined(SQLITE_OMIT_PAGER_PRAGMAS)
  {/* zName:     */ "cache_size",
@@ -120727,11 +123856,13 @@ static const PragmaName aPragmaName[] = {
   /* ColNames:  */ 0, 0,
   /* iArg:      */ 0 },
 #endif
+#if !defined(SQLITE_OMIT_CASE_SENSITIVE_LIKE_PRAGMA)
  {/* zName:     */ "case_sensitive_like",
   /* ePragTyp:  */ PragTyp_CASE_SENSITIVE_LIKE,
   /* ePragFlg:  */ PragFlg_NoColumns,
   /* ColNames:  */ 0, 0,
   /* iArg:      */ 0 },
+#endif
  {/* zName:     */ "cell_size_check",
   /* ePragTyp:  */ PragTyp_FLAG,
   /* ePragFlg:  */ PragFlg_Result0|PragFlg_NoColumns1,
@@ -120748,7 +123879,7 @@ static const PragmaName aPragmaName[] = {
  {/* zName:     */ "collation_list",
   /* ePragTyp:  */ PragTyp_COLLATION_LIST,
   /* ePragFlg:  */ PragFlg_Result0,
-  /* ColNames:  */ 26, 2,
+  /* ColNames:  */ 32, 2,
   /* iArg:      */ 0 },
 #endif
 #if !defined(SQLITE_OMIT_COMPILEOPTION_DIAGS)
@@ -120783,14 +123914,14 @@ static const PragmaName aPragmaName[] = {
  {/* zName:     */ "database_list",
   /* ePragTyp:  */ PragTyp_DATABASE_LIST,
   /* ePragFlg:  */ PragFlg_NeedSchema|PragFlg_Result0,
-  /* ColNames:  */ 35, 3,
+  /* ColNames:  */ 41, 3,
   /* iArg:      */ 0 },
 #endif
 #if !defined(SQLITE_OMIT_PAGER_PRAGMAS) && !defined(SQLITE_OMIT_DEPRECATED)
  {/* zName:     */ "default_cache_size",
   /* ePragTyp:  */ PragTyp_DEFAULT_CACHE_SIZE,
   /* ePragFlg:  */ PragFlg_NeedSchema|PragFlg_Result0|PragFlg_SchemaReq|PragFlg_NoColumns1,
-  /* ColNames:  */ 45, 1,
+  /* ColNames:  */ 49, 1,
   /* iArg:      */ 0 },
 #endif
 #if !defined(SQLITE_OMIT_FLAG_PRAGMAS)
@@ -120820,7 +123951,7 @@ static const PragmaName aPragmaName[] = {
  {/* zName:     */ "foreign_key_check",
   /* ePragTyp:  */ PragTyp_FOREIGN_KEY_CHECK,
   /* ePragFlg:  */ PragFlg_NeedSchema|PragFlg_Result0,
-  /* ColNames:  */ 31, 4,
+  /* ColNames:  */ 37, 4,
   /* iArg:      */ 0 },
 #endif
 #if !defined(SQLITE_OMIT_FOREIGN_KEY)
@@ -120859,22 +123990,27 @@ static const PragmaName aPragmaName[] = {
   /* iArg:      */ SQLITE_FullFSync },
 #endif
 #if !defined(SQLITE_OMIT_SCHEMA_PRAGMAS)
-#if defined(SQLITE_INTROSPECTION_PRAGMAS)
+#if !defined(SQLITE_OMIT_INTROSPECTION_PRAGMAS)
  {/* zName:     */ "function_list",
   /* ePragTyp:  */ PragTyp_FUNCTION_LIST,
   /* ePragFlg:  */ PragFlg_Result0,
-  /* ColNames:  */ 41, 2,
+  /* ColNames:  */ 21, 6,
   /* iArg:      */ 0 },
 #endif
 #endif
+ {/* zName:     */ "hard_heap_limit",
+  /* ePragTyp:  */ PragTyp_HARD_HEAP_LIMIT,
+  /* ePragFlg:  */ PragFlg_Result0,
+  /* ColNames:  */ 0, 0,
+  /* iArg:      */ 0 },
 #if defined(SQLITE_HAS_CODEC)
  {/* zName:     */ "hexkey",
-  /* ePragTyp:  */ PragTyp_HEXKEY,
+  /* ePragTyp:  */ PragTyp_KEY,
   /* ePragFlg:  */ 0,
   /* ColNames:  */ 0, 0,
   /* iArg:      */ 2 },
  {/* zName:     */ "hexrekey",
-  /* ePragTyp:  */ PragTyp_HEXKEY,
+  /* ePragTyp:  */ PragTyp_KEY,
   /* ePragFlg:  */ 0,
   /* ColNames:  */ 0, 0,
   /* iArg:      */ 3 },
@@ -120904,7 +124040,7 @@ static const PragmaName aPragmaName[] = {
  {/* zName:     */ "index_list",
   /* ePragTyp:  */ PragTyp_INDEX_LIST,
   /* ePragFlg:  */ PragFlg_NeedSchema|PragFlg_Result1|PragFlg_SchemaOpt,
-  /* ColNames:  */ 26, 5,
+  /* ColNames:  */ 32, 5,
   /* iArg:      */ 0 },
  {/* zName:     */ "index_xinfo",
   /* ePragTyp:  */ PragTyp_INDEX_INFO,
@@ -120944,11 +124080,6 @@ static const PragmaName aPragmaName[] = {
   /* ePragFlg:  */ PragFlg_Result0|PragFlg_NoColumns1,
   /* ColNames:  */ 0, 0,
   /* iArg:      */ SQLITE_LegacyAlter },
- {/* zName:     */ "legacy_file_format",
-  /* ePragTyp:  */ PragTyp_FLAG,
-  /* ePragFlg:  */ PragFlg_Result0|PragFlg_NoColumns1,
-  /* ColNames:  */ 0, 0,
-  /* iArg:      */ SQLITE_LegacyFileFmt },
 #endif
 #if !defined(SQLITE_OMIT_PAGER_PRAGMAS) && SQLITE_ENABLE_LOCKING_STYLE
  {/* zName:     */ "lock_proxy_file",
@@ -120961,7 +124092,7 @@ static const PragmaName aPragmaName[] = {
  {/* zName:     */ "lock_status",
   /* ePragTyp:  */ PragTyp_LOCK_STATUS,
   /* ePragFlg:  */ PragFlg_Result0,
-  /* ColNames:  */ 43, 2,
+  /* ColNames:  */ 47, 2,
   /* iArg:      */ 0 },
 #endif
 #if !defined(SQLITE_OMIT_PAGER_PRAGMAS)
@@ -120983,7 +124114,7 @@ static const PragmaName aPragmaName[] = {
 #endif
 #if !defined(SQLITE_OMIT_SCHEMA_PRAGMAS)
 #if !defined(SQLITE_OMIT_VIRTUALTABLE)
-#if defined(SQLITE_INTROSPECTION_PRAGMAS)
+#if !defined(SQLITE_OMIT_INTROSPECTION_PRAGMAS)
  {/* zName:     */ "module_list",
   /* ePragTyp:  */ PragTyp_MODULE_LIST,
   /* ePragFlg:  */ PragFlg_Result0,
@@ -121018,7 +124149,7 @@ static const PragmaName aPragmaName[] = {
   /* iArg:      */ SQLITE_ParserTrace },
 #endif
 #endif
-#if defined(SQLITE_INTROSPECTION_PRAGMAS)
+#if !defined(SQLITE_OMIT_INTROSPECTION_PRAGMAS)
  {/* zName:     */ "pragma_list",
   /* ePragTyp:  */ PragTyp_PRAGMA_LIST,
   /* ePragFlg:  */ PragFlg_Result0,
@@ -121109,7 +124240,7 @@ static const PragmaName aPragmaName[] = {
  {/* zName:     */ "stats",
   /* ePragTyp:  */ PragTyp_STATS,
   /* ePragFlg:  */ PragFlg_NeedSchema|PragFlg_Result0|PragFlg_SchemaReq,
-  /* ColNames:  */ 21, 5,
+  /* ColNames:  */ 27, 5,
   /* iArg:      */ 0 },
 #endif
 #if !defined(SQLITE_OMIT_PAGER_PRAGMAS)
@@ -121160,6 +124291,13 @@ static const PragmaName aPragmaName[] = {
   /* ePragFlg:  */ PragFlg_Result0,
   /* ColNames:  */ 0, 0,
   /* iArg:      */ 0 },
+#if !defined(SQLITE_OMIT_FLAG_PRAGMAS)
+ {/* zName:     */ "trusted_schema",
+  /* ePragTyp:  */ PragTyp_FLAG,
+  /* ePragFlg:  */ PragFlg_Result0|PragFlg_NoColumns1,
+  /* ColNames:  */ 0, 0,
+  /* iArg:      */ SQLITE_TrustedSchema },
+#endif
 #if !defined(SQLITE_OMIT_SCHEMA_VERSION_PRAGMAS)
  {/* zName:     */ "user_version",
   /* ePragTyp:  */ PragTyp_HEADER_VALUE,
@@ -121205,7 +124343,7 @@ static const PragmaName aPragmaName[] = {
  {/* zName:     */ "wal_checkpoint",
   /* ePragTyp:  */ PragTyp_WAL_CHECKPOINT,
   /* ePragFlg:  */ PragFlg_NeedSchema,
-  /* ColNames:  */ 38, 3,
+  /* ColNames:  */ 44, 3,
   /* iArg:      */ 0 },
 #endif
 #if !defined(SQLITE_OMIT_FLAG_PRAGMAS)
@@ -121216,7 +124354,7 @@ static const PragmaName aPragmaName[] = {
   /* iArg:      */ SQLITE_WriteSchema|SQLITE_NoSchemaError },
 #endif
 };
-/* Number of pragmas: 62 on by default, 81 total. */
+/* Number of pragmas: 66 on by default, 82 total. */
 
 /************** End of pragma.h **********************************************/
 /************** Continuing where we left off in pragma.c *********************/
@@ -121486,6 +124624,55 @@ static const PragmaName *pragmaLocate(const char *zName){
   return lwr>upr ? 0 : &aPragmaName[mid];
 }
 
+/*
+** Create zero or more entries in the output for the SQL functions
+** defined by FuncDef p.
+*/
+static void pragmaFunclistLine(
+  Vdbe *v,               /* The prepared statement being created */
+  FuncDef *p,            /* A particular function definition */
+  int isBuiltin,         /* True if this is a built-in function */
+  int showInternFuncs    /* True if showing internal functions */
+){
+  for(; p; p=p->pNext){
+    const char *zType;
+    static const u32 mask = 
+        SQLITE_DETERMINISTIC |
+        SQLITE_DIRECTONLY |
+        SQLITE_SUBTYPE |
+        SQLITE_INNOCUOUS |
+        SQLITE_FUNC_INTERNAL
+    ;
+    static const char *azEnc[] = { 0, "utf8", "utf16le", "utf16be" };
+
+    assert( SQLITE_FUNC_ENCMASK==0x3 );
+    assert( strcmp(azEnc[SQLITE_UTF8],"utf8")==0 );
+    assert( strcmp(azEnc[SQLITE_UTF16LE],"utf16le")==0 );
+    assert( strcmp(azEnc[SQLITE_UTF16BE],"utf16be")==0 );
+
+    if( p->xSFunc==0 ) continue;
+    if( (p->funcFlags & SQLITE_FUNC_INTERNAL)!=0
+     && showInternFuncs==0
+    ){
+      continue;
+    }    
+    if( p->xValue!=0 ){
+      zType = "w";
+    }else if( p->xFinalize!=0 ){
+      zType = "a";
+    }else{
+      zType = "s";
+    }
+    sqlite3VdbeMultiLoad(v, 1, "sissii",
+       p->zName, isBuiltin,
+       zType, azEnc[p->funcFlags&SQLITE_FUNC_ENCMASK],
+       p->nArg,
+       (p->funcFlags & mask) ^ SQLITE_INNOCUOUS
+    );
+  }
+}
+
+
 /*
 ** Helper subroutine for PRAGMA integrity_check:
 **
@@ -121835,6 +125022,11 @@ SQLITE_PRIVATE void sqlite3Pragma(
         ** then do a query */
         eMode = PAGER_JOURNALMODE_QUERY;
       }
+      if( eMode==PAGER_JOURNALMODE_OFF && (db->flags & SQLITE_Defensive)!=0 ){
+        /* Do not allow journal-mode "OFF" in defensive since the database
+        ** can become corrupted using ordinary SQL when the journal is off */
+        eMode = PAGER_JOURNALMODE_QUERY;
+      }
     }
     if( eMode==PAGER_JOURNALMODE_QUERY && pId2->n==0 ){
       /* Convert "PRAGMA journal_mode" into "PRAGMA main.journal_mode" */
@@ -122286,10 +125478,19 @@ SQLITE_PRIVATE void sqlite3Pragma(
       sqlite3CodeVerifySchema(pParse, iTabDb);
       sqlite3ViewGetColumnNames(pParse, pTab);
       for(i=0, pCol=pTab->aCol; i<pTab->nCol; i++, pCol++){
-        int isHidden = IsHiddenColumn(pCol);
-        if( isHidden && pPragma->iArg==0 ){
-          nHidden++;
-          continue;
+        int isHidden = 0;
+        if( pCol->colFlags & COLFLAG_NOINSERT ){
+          if( pPragma->iArg==0 ){
+            nHidden++;
+            continue;
+          }
+          if( pCol->colFlags & COLFLAG_VIRTUAL ){
+            isHidden = 2;  /* GENERATED ALWAYS AS ... VIRTUAL */
+          }else if( pCol->colFlags & COLFLAG_STORED ){
+            isHidden = 3;  /* GENERATED ALWAYS AS ... STORED */
+          }else{ assert( pCol->colFlags & COLFLAG_HIDDEN );
+            isHidden = 1;  /* HIDDEN */
+          }
         }
         if( (pCol->colFlags & COLFLAG_PRIMKEY)==0 ){
           k = 0;
@@ -122298,13 +125499,13 @@ SQLITE_PRIVATE void sqlite3Pragma(
         }else{
           for(k=1; k<=pTab->nCol && pPk->aiColumn[k-1]!=i; k++){}
         }
-        assert( pCol->pDflt==0 || pCol->pDflt->op==TK_SPAN );
+        assert( pCol->pDflt==0 || pCol->pDflt->op==TK_SPAN || isHidden>=2 );
         sqlite3VdbeMultiLoad(v, 1, pPragma->iArg ? "issisii" : "issisi",
                i-nHidden,
                pCol->zName,
                sqlite3ColumnType(pCol,""),
                pCol->notNull ? 1 : 0,
-               pCol->pDflt ? pCol->pDflt->u.zToken : 0,
+               pCol->pDflt && isHidden<2 ? pCol->pDflt->u.zToken : 0,
                k,
                isHidden);
       }
@@ -122343,6 +125544,15 @@ SQLITE_PRIVATE void sqlite3Pragma(
     Index *pIdx;
     Table *pTab;
     pIdx = sqlite3FindIndex(db, zRight, zDb);
+    if( pIdx==0 ){
+      /* If there is no index named zRight, check to see if there is a
+      ** WITHOUT ROWID table named zRight, and if there is, show the
+      ** structure of the PRIMARY KEY index for that table. */
+      pTab = sqlite3LocateTable(pParse, LOCATE_NOERR, zRight, zDb);
+      if( pTab && !HasRowid(pTab) ){
+        pIdx = sqlite3PrimaryKeyIndex(pTab);
+      }
+    }
     if( pIdx ){
       int iIdxDb = sqlite3SchemaToIndex(db, pIdx->pSchema);
       int i;
@@ -122422,21 +125632,21 @@ SQLITE_PRIVATE void sqlite3Pragma(
   }
   break;
 
-#ifdef SQLITE_INTROSPECTION_PRAGMAS
+#ifndef SQLITE_OMIT_INTROSPECTION_PRAGMAS
   case PragTyp_FUNCTION_LIST: {
     int i;
     HashElem *j;
     FuncDef *p;
-    pParse->nMem = 2;
+    int showInternFunc = (db->mDbFlags & DBFLAG_InternalFunc)!=0;
+    pParse->nMem = 6;
     for(i=0; i<SQLITE_FUNC_HASH_SZ; i++){
       for(p=sqlite3BuiltinFunctions.a[i]; p; p=p->u.pHash ){
-        if( p->funcFlags & SQLITE_FUNC_INTERNAL ) continue;
-        sqlite3VdbeMultiLoad(v, 1, "si", p->zName, 1);
+        pragmaFunclistLine(v, p, 1, showInternFunc);
       }
     }
     for(j=sqliteHashFirst(&db->aFunc); j; j=sqliteHashNext(j)){
       p = (FuncDef*)sqliteHashData(j);
-      sqlite3VdbeMultiLoad(v, 1, "si", p->zName, 0);
+      pragmaFunclistLine(v, p, 0, showInternFunc);
     }
   }
   break;
@@ -122612,6 +125822,7 @@ SQLITE_PRIVATE void sqlite3Pragma(
 #endif /* !defined(SQLITE_OMIT_TRIGGER) */
 #endif /* !defined(SQLITE_OMIT_FOREIGN_KEY) */
 
+#ifndef SQLITE_OMIT_CASE_SENSITIVE_LIKE_PRAGMA
   /* Reinstall the LIKE and GLOB functions.  The variant of LIKE
   ** used will be case sensitive or not depending on the RHS.
   */
@@ -122621,6 +125832,7 @@ SQLITE_PRIVATE void sqlite3Pragma(
     }
   }
   break;
+#endif /* SQLITE_OMIT_CASE_SENSITIVE_LIKE_PRAGMA */
 
 #ifndef SQLITE_INTEGRITY_CHECK_ERROR_MAX
 # define SQLITE_INTEGRITY_CHECK_ERROR_MAX 100
@@ -122752,7 +125964,7 @@ SQLITE_PRIVATE void sqlite3Pragma(
         loopTop = sqlite3VdbeAddOp2(v, OP_AddImm, 7, 1);
         if( !isQuick ){
           /* Sanity check on record header decoding */
-          sqlite3VdbeAddOp3(v, OP_Column, iDataCur, pTab->nCol-1, 3);
+          sqlite3VdbeAddOp3(v, OP_Column, iDataCur, pTab->nNVCol-1,3);
           sqlite3VdbeChangeP5(v, OPFLAG_TYPEOFARG);
         }
         /* Verify that all NOT NULL columns really are NOT NULL */
@@ -122762,7 +125974,9 @@ SQLITE_PRIVATE void sqlite3Pragma(
           if( j==pTab->iPKey ) continue;
           if( pTab->aCol[j].notNull==0 ) continue;
           sqlite3ExprCodeGetColumnOfTable(v, pTab, iDataCur, j, 3);
-          sqlite3VdbeChangeP5(v, OPFLAG_TYPEOFARG);
+          if( sqlite3VdbeGetOp(v,-1)->opcode==OP_Column ){
+            sqlite3VdbeChangeP5(v, OPFLAG_TYPEOFARG);
+          }
           jmp2 = sqlite3VdbeAddOp1(v, OP_NotNull, 3); VdbeCoverage(v);
           zErr = sqlite3MPrintf(db, "NULL value in %s.%s", pTab->zName,
                               pTab->aCol[j].zName);
@@ -122939,10 +126153,17 @@ SQLITE_PRIVATE void sqlite3Pragma(
       ** will be overwritten when the schema is next loaded. If it does not
       ** already exists, it will be created to use the new encoding value.
       */
-      if( 
-        !(DbHasProperty(db, 0, DB_SchemaLoaded)) || 
-        DbHasProperty(db, 0, DB_Empty) 
-      ){
+      int canChangeEnc = 1;  /* True if allowed to change the encoding */
+      int i;                 /* For looping over all attached databases */
+      for(i=0; i<db->nDb; i++){
+        if( db->aDb[i].pBt!=0
+         && DbHasProperty(db,i,DB_SchemaLoaded)
+         && !DbHasProperty(db,i,DB_Empty)
+        ){
+          canChangeEnc = 0;
+        }
+      }
+      if( canChangeEnc ){
         for(pEnc=&encnames[0]; pEnc->zName; pEnc++){
           if( 0==sqlite3StrICmp(zRight, pEnc->zName) ){
             SCHEMA_ENC(db) = ENC(db) =
@@ -123255,6 +126476,27 @@ SQLITE_PRIVATE void sqlite3Pragma(
     break;
   }
 
+  /*
+  **   PRAGMA hard_heap_limit
+  **   PRAGMA hard_heap_limit = N
+  **
+  ** Invoke sqlite3_hard_heap_limit64() to query or set the hard heap
+  ** limit.  The hard heap limit can be activated or lowered by this
+  ** pragma, but not raised or deactivated.  Only the
+  ** sqlite3_hard_heap_limit64() C-language API can raise or deactivate
+  ** the hard heap limit.  This allows an application to set a heap limit
+  ** constraint that cannot be relaxed by an untrusted SQL script.
+  */
+  case PragTyp_HARD_HEAP_LIMIT: {
+    sqlite3_int64 N;
+    if( zRight && sqlite3DecOrHexToI64(zRight, &N)==SQLITE_OK ){
+      sqlite3_int64 iPrior = sqlite3_hard_heap_limit64(-1);
+      if( N>0 && (iPrior==0 || iPrior>N) ) sqlite3_hard_heap_limit64(N);
+    }
+    returnSingleInt(v, sqlite3_hard_heap_limit64(-1));
+    break;
+  }
+
   /*
   **   PRAGMA threads
   **   PRAGMA threads = N
@@ -123314,28 +126556,30 @@ SQLITE_PRIVATE void sqlite3Pragma(
   */
   case PragTyp_KEY: {
     if( zRight ){
-      int n = pPragma->iArg<4 ? sqlite3Strlen30(zRight) : -1;
-      if( (pPragma->iArg & 1)==0 ){
-        sqlite3_key_v2(db, zDb, zRight, n);
+      char zBuf[40];
+      const char *zKey = zRight;
+      int n;
+      if( pPragma->iArg==2 || pPragma->iArg==3 ){
+        u8 iByte;
+        int i;
+        for(i=0, iByte=0; i<sizeof(zBuf)*2 && sqlite3Isxdigit(zRight[i]); i++){
+          iByte = (iByte<<4) + sqlite3HexToInt(zRight[i]);
+          if( (i&1)!=0 ) zBuf[i/2] = iByte;
+        }
+        zKey = zBuf;
+        n = i/2;
       }else{
-        sqlite3_rekey_v2(db, zDb, zRight, n);
-      }
-    }
-    break;
-  }
-  case PragTyp_HEXKEY: {
-    if( zRight ){
-      u8 iByte;
-      int i;
-      char zKey[40];
-      for(i=0, iByte=0; i<sizeof(zKey)*2 && sqlite3Isxdigit(zRight[i]); i++){
-        iByte = (iByte<<4) + sqlite3HexToInt(zRight[i]);
-        if( (i&1)!=0 ) zKey[i/2] = iByte;
+        n = pPragma->iArg<4 ? sqlite3Strlen30(zRight) : -1;
       }
       if( (pPragma->iArg & 1)==0 ){
-        sqlite3_key_v2(db, zDb, zKey, i/2);
+        rc = sqlite3_key_v2(db, zDb, zKey, n);
       }else{
-        sqlite3_rekey_v2(db, zDb, zKey, i/2);
+        rc = sqlite3_rekey_v2(db, zDb, zKey, n);
+      }
+      if( rc==SQLITE_OK && n!=0 ){
+        sqlite3VdbeSetNumCols(v, 1);
+        sqlite3VdbeSetColName(v, 0, COLNAME_NAME, "ok", SQLITE_STATIC);
+        returnSingleText(v, "ok");
       }
     }
     break;
@@ -123746,6 +126990,18 @@ SQLITE_PRIVATE int sqlite3IndexHasDuplicateRootPage(Index *pIndex){
   return 0;
 }
 
+/* forward declaration */
+static int sqlite3Prepare(
+  sqlite3 *db,              /* Database handle. */
+  const char *zSql,         /* UTF-8 encoded SQL statement. */
+  int nBytes,               /* Length of zSql in bytes. */
+  u32 prepFlags,            /* Zero or more SQLITE_PREPARE_* flags */
+  Vdbe *pReprepare,         /* VM being reprepared */
+  sqlite3_stmt **ppStmt,    /* OUT: A pointer to the prepared statement */
+  const char **pzTail       /* OUT: End of parsed string */
+);
+
+
 /*
 ** This is the callback routine for the code that initializes the
 ** database.  See sqlite3Init() below for additional information.
@@ -123753,9 +127009,11 @@ SQLITE_PRIVATE int sqlite3IndexHasDuplicateRootPage(Index *pIndex){
 **
 ** Each callback contains the following information:
 **
-**     argv[0] = name of thing being created
-**     argv[1] = root page number for table or index. 0 for trigger or view.
-**     argv[2] = SQL text for the CREATE statement.
+**     argv[0] = type of object: "table", "index", "trigger", or "view".
+**     argv[1] = name of thing being created
+**     argv[2] = associated table if an index or trigger
+**     argv[3] = root page number for table or index. 0 for trigger or view.
+**     argv[4] = SQL text for the CREATE statement.
 **
 */
 SQLITE_PRIVATE int sqlite3InitCallback(void *pInit, int argc, char **argv, char **NotUsed){
@@ -123763,21 +127021,21 @@ SQLITE_PRIVATE int sqlite3InitCallback(void *pInit, int argc, char **argv, char
   sqlite3 *db = pData->db;
   int iDb = pData->iDb;
 
-  assert( argc==3 );
+  assert( argc==5 );
   UNUSED_PARAMETER2(NotUsed, argc);
   assert( sqlite3_mutex_held(db->mutex) );
   DbClearProperty(db, iDb, DB_Empty);
   pData->nInitRow++;
   if( db->mallocFailed ){
-    corruptSchema(pData, argv[0], 0);
+    corruptSchema(pData, argv[1], 0);
     return 1;
   }
 
   assert( iDb>=0 && iDb<db->nDb );
   if( argv==0 ) return 0;   /* Might happen if EMPTY_RESULT_CALLBACKS are on */
-  if( argv[1]==0 ){
-    corruptSchema(pData, argv[0], 0);
-  }else if( sqlite3_strnicmp(argv[2],"create ",7)==0 ){
+  if( argv[3]==0 ){
+    corruptSchema(pData, argv[1], 0);
+  }else if( sqlite3_strnicmp(argv[4],"create ",7)==0 ){
     /* Call the parser to process a CREATE TABLE, INDEX or VIEW.
     ** But because db->init.busy is set to 1, no VDBE code is generated
     ** or executed.  All the parser does is build the internal data
@@ -123790,9 +127048,11 @@ SQLITE_PRIVATE int sqlite3InitCallback(void *pInit, int argc, char **argv, char
 
     assert( db->init.busy );
     db->init.iDb = iDb;
-    db->init.newTnum = sqlite3Atoi(argv[1]);
+    db->init.newTnum = sqlite3Atoi(argv[3]);
     db->init.orphanTrigger = 0;
-    TESTONLY(rcp = ) sqlite3_prepare(db, argv[2], -1, &pStmt, 0);
+    db->init.azInit = argv;
+    pStmt = 0;
+    TESTONLY(rcp = ) sqlite3Prepare(db, argv[4], -1, 0, 0, &pStmt, 0);
     rc = db->errCode;
     assert( (rc&0xFF)==(rcp&0xFF) );
     db->init.iDb = saved_iDb;
@@ -123801,17 +127061,17 @@ SQLITE_PRIVATE int sqlite3InitCallback(void *pInit, int argc, char **argv, char
       if( db->init.orphanTrigger ){
         assert( iDb==1 );
       }else{
-        pData->rc = rc;
+        if( rc > pData->rc ) pData->rc = rc;
         if( rc==SQLITE_NOMEM ){
           sqlite3OomFault(db);
         }else if( rc!=SQLITE_INTERRUPT && (rc&0xFF)!=SQLITE_LOCKED ){
-          corruptSchema(pData, argv[0], sqlite3_errmsg(db));
+          corruptSchema(pData, argv[1], sqlite3_errmsg(db));
         }
       }
     }
     sqlite3_finalize(pStmt);
-  }else if( argv[0]==0 || (argv[2]!=0 && argv[2][0]!=0) ){
-    corruptSchema(pData, argv[0], 0);
+  }else if( argv[1]==0 || (argv[4]!=0 && argv[4][0]!=0) ){
+    corruptSchema(pData, argv[1], 0);
   }else{
     /* If the SQL column is blank it means this is an index that
     ** was created to be the PRIMARY KEY or to fulfill a UNIQUE
@@ -123820,13 +127080,13 @@ SQLITE_PRIVATE int sqlite3InitCallback(void *pInit, int argc, char **argv, char
     ** to do here is record the root page number for that index.
     */
     Index *pIndex;
-    pIndex = sqlite3FindIndex(db, argv[0], db->aDb[iDb].zDbSName);
+    pIndex = sqlite3FindIndex(db, argv[1], db->aDb[iDb].zDbSName);
     if( pIndex==0
-     || sqlite3GetInt32(argv[1],&pIndex->tnum)==0
+     || sqlite3GetInt32(argv[3],&pIndex->tnum)==0
      || pIndex->tnum<2
      || sqlite3IndexHasDuplicateRootPage(pIndex)
     ){
-      corruptSchema(pData, argv[0], pIndex?"invalid rootpage":"orphan index");
+      corruptSchema(pData, argv[1], pIndex?"invalid rootpage":"orphan index");
     }
   }
   return 0;
@@ -123847,7 +127107,7 @@ SQLITE_PRIVATE int sqlite3InitOne(sqlite3 *db, int iDb, char **pzErrMsg, u32 mFl
   int size;
 #endif
   Db *pDb;
-  char const *azArg[4];
+  char const *azArg[6];
   int meta[5];
   InitData initData;
   const char *zMasterName;
@@ -123866,18 +127126,20 @@ SQLITE_PRIVATE int sqlite3InitOne(sqlite3 *db, int iDb, char **pzErrMsg, u32 mFl
   ** table name will be inserted automatically by the parser so we can just
   ** use the abbreviation "x" here.  The parser will also automatically tag
   ** the schema table as read-only. */
-  azArg[0] = zMasterName = SCHEMA_TABLE(iDb);
-  azArg[1] = "1";
-  azArg[2] = "CREATE TABLE x(type text,name text,tbl_name text,"
+  azArg[0] = "table";
+  azArg[1] = zMasterName = SCHEMA_TABLE(iDb);
+  azArg[2] = azArg[1];
+  azArg[3] = "1";
+  azArg[4] = "CREATE TABLE x(type text,name text,tbl_name text,"
                             "rootpage int,sql text)";
-  azArg[3] = 0;
+  azArg[5] = 0;
   initData.db = db;
   initData.iDb = iDb;
   initData.rc = SQLITE_OK;
   initData.pzErrMsg = pzErrMsg;
   initData.mInitFlags = mFlags;
   initData.nInitRow = 0;
-  sqlite3InitCallback(&initData, 3, (char **)azArg, 0);
+  sqlite3InitCallback(&initData, 5, (char **)azArg, 0);
   if( initData.rc ){
     rc = initData.rc;
     goto error_out;
@@ -124003,7 +127265,7 @@ SQLITE_PRIVATE int sqlite3InitOne(sqlite3 *db, int iDb, char **pzErrMsg, u32 mFl
   {
     char *zSql;
     zSql = sqlite3MPrintf(db, 
-        "SELECT name, rootpage, sql FROM \"%w\".%s ORDER BY rowid",
+        "SELECT*FROM\"%w\".%s ORDER BY rowid",
         db->aDb[iDb].zDbSName, zMasterName);
 #ifndef SQLITE_OMIT_AUTHORIZATION
     {
@@ -124211,6 +127473,7 @@ SQLITE_PRIVATE void sqlite3ParserReset(Parse *pParse){
   if( db ){
     assert( db->lookaside.bDisable >= pParse->disableLookaside );
     db->lookaside.bDisable -= pParse->disableLookaside;
+    db->lookaside.sz = db->lookaside.bDisable ? 0 : db->lookaside.szTrue;
   }
   pParse->disableLookaside = 0;
 }
@@ -124244,7 +127507,7 @@ static int sqlite3Prepare(
   */
   if( prepFlags & SQLITE_PREPARE_PERSISTENT ){
     sParse.disableLookaside++;
-    db->lookaside.bDisable++;
+    DisableLookaside;
   }
   sParse.disableVtab = (prepFlags & SQLITE_PREPARE_NO_VTAB)!=0;
 
@@ -124271,16 +127534,18 @@ static int sqlite3Prepare(
   ** but it does *not* override schema lock detection, so this all still
   ** works even if READ_UNCOMMITTED is set.
   */
-  for(i=0; i<db->nDb; i++) {
-    Btree *pBt = db->aDb[i].pBt;
-    if( pBt ){
-      assert( sqlite3BtreeHoldsMutex(pBt) );
-      rc = sqlite3BtreeSchemaLocked(pBt);
-      if( rc ){
-        const char *zDb = db->aDb[i].zDbSName;
-        sqlite3ErrorWithMsg(db, rc, "database schema is locked: %s", zDb);
-        testcase( db->flags & SQLITE_ReadUncommit );
-        goto end_prepare;
+  if( !db->noSharedCache ){
+    for(i=0; i<db->nDb; i++) {
+      Btree *pBt = db->aDb[i].pBt;
+      if( pBt ){
+        assert( sqlite3BtreeHoldsMutex(pBt) );
+        rc = sqlite3BtreeSchemaLocked(pBt);
+        if( rc ){
+          const char *zDb = db->aDb[i].zDbSName;
+          sqlite3ErrorWithMsg(db, rc, "database schema is locked: %s", zDb);
+          testcase( db->flags & SQLITE_ReadUncommit );
+          goto end_prepare;
+        }
       }
     }
   }
@@ -124311,46 +127576,25 @@ static int sqlite3Prepare(
   }
   assert( 0==sParse.nQueryLoop );
 
-  if( sParse.rc==SQLITE_DONE ) sParse.rc = SQLITE_OK;
+  if( sParse.rc==SQLITE_DONE ){
+    sParse.rc = SQLITE_OK;
+  }
   if( sParse.checkSchema ){
     schemaIsValid(&sParse);
   }
-  if( db->mallocFailed ){
-    sParse.rc = SQLITE_NOMEM_BKPT;
-  }
   if( pzTail ){
     *pzTail = sParse.zTail;
   }
-  rc = sParse.rc;
-
-#ifndef SQLITE_OMIT_EXPLAIN
-  if( rc==SQLITE_OK && sParse.pVdbe && sParse.explain ){
-    static const char * const azColName[] = {
-       "addr", "opcode", "p1", "p2", "p3", "p4", "p5", "comment",
-       "id", "parent", "notused", "detail"
-    };
-    int iFirst, mx;
-    if( sParse.explain==2 ){
-      sqlite3VdbeSetNumCols(sParse.pVdbe, 4);
-      iFirst = 8;
-      mx = 12;
-    }else{
-      sqlite3VdbeSetNumCols(sParse.pVdbe, 8);
-      iFirst = 0;
-      mx = 8;
-    }
-    for(i=iFirst; i<mx; i++){
-      sqlite3VdbeSetColName(sParse.pVdbe, i-iFirst, COLNAME_NAME,
-                            azColName[i], SQLITE_STATIC);
-    }
-  }
-#endif
 
   if( db->init.busy==0 ){
     sqlite3VdbeSetSql(sParse.pVdbe, zSql, (int)(sParse.zTail-zSql), prepFlags);
   }
-  if( sParse.pVdbe && (rc!=SQLITE_OK || db->mallocFailed) ){
-    sqlite3VdbeFinalize(sParse.pVdbe);
+  if( db->mallocFailed ){
+    sParse.rc = SQLITE_NOMEM_BKPT;
+  }
+  rc = sParse.rc;
+  if( rc!=SQLITE_OK ){
+    if( sParse.pVdbe ) sqlite3VdbeFinalize(sParse.pVdbe);
     assert(!(*ppStmt));
   }else{
     *ppStmt = (sqlite3_stmt*)sParse.pVdbe;
@@ -124705,7 +127949,10 @@ struct SortCtx {
 
 /*
 ** Delete all the content of a Select structure.  Deallocate the structure
-** itself only if bFree is true.
+** itself depending on the value of bFree
+**
+** If bFree==1, call sqlite3DbFree() on the p object.
+** If bFree==0, Leave the first Select object unfreed
 */
 static void clearSelect(sqlite3 *db, Select *p, int bFree){
   while( p ){
@@ -124721,6 +127968,7 @@ static void clearSelect(sqlite3 *db, Select *p, int bFree){
     if( OK_IF_ALWAYS_TRUE(p->pWinDefn) ){
       sqlite3WindowListDelete(db, p->pWinDefn);
     }
+    assert( p->pWin==0 );
 #endif
     if( OK_IF_ALWAYS_TRUE(p->pWith) ) sqlite3WithDelete(db, p->pWith);
     if( bFree ) sqlite3DbFreeNN(db, p);
@@ -124808,6 +128056,21 @@ SQLITE_PRIVATE void sqlite3SelectDelete(sqlite3 *db, Select *p){
   if( OK_IF_ALWAYS_TRUE(p) ) clearSelect(db, p, 1);
 }
 
+/*
+** Delete all the substructure for p, but keep p allocated.  Redefine
+** p to be a single SELECT where every column of the result set has a
+** value of NULL.
+*/
+SQLITE_PRIVATE void sqlite3SelectReset(Parse *pParse, Select *p){
+  if( ALWAYS(p) ){
+    clearSelect(pParse->db, p, 0);
+    memset(&p->iLimit, 0, sizeof(Select) - offsetof(Select,iLimit));
+    p->pEList = sqlite3ExprListAppend(pParse, 0,
+                     sqlite3ExprAlloc(pParse->db,TK_NULL,0,0));
+    p->pSrc = sqlite3DbMallocZero(pParse->db, sizeof(SrcList));
+  }
+}
+
 /*
 ** Return a pointer to the right-most SELECT statement in a compound.
 */
@@ -124916,7 +128179,8 @@ static int tableAndColumnIndex(
   int N,               /* Number of tables in pSrc->a[] to search */
   const char *zCol,    /* Name of the column we are looking for */
   int *piTab,          /* Write index of pSrc->a[] here */
-  int *piCol           /* Write index of pSrc->a[*piTab].pTab->aCol[] here */
+  int *piCol,          /* Write index of pSrc->a[*piTab].pTab->aCol[] here */
+  int bIgnoreHidden    /* True to ignore hidden columns */
 ){
   int i;               /* For looping over tables in pSrc */
   int iCol;            /* Index of column matching zCol */
@@ -124924,7 +128188,9 @@ static int tableAndColumnIndex(
   assert( (piTab==0)==(piCol==0) );  /* Both or neither are NULL */
   for(i=0; i<N; i++){
     iCol = columnIndex(pSrc->a[i].pTab, zCol);
-    if( iCol>=0 ){
+    if( iCol>=0 
+     && (bIgnoreHidden==0 || IsHiddenColumn(&pSrc->a[i].pTab->aCol[iCol])==0)
+    ){
       if( piTab ){
         *piTab = i;
         *piCol = iCol;
@@ -124976,7 +128242,7 @@ static void addWhereTerm(
     ExprSetVVAProperty(pEq, EP_NoReduce);
     pEq->iRightJoinTable = (i16)pE2->iTable;
   }
-  *ppWhere = sqlite3ExprAnd(db, *ppWhere, pEq);
+  *ppWhere = sqlite3ExprAnd(pParse, *ppWhere, pEq);
 }
 
 /*
@@ -125005,7 +128271,7 @@ static void addWhereTerm(
 ** after the t1 loop and rows with t1.x!=5 will never appear in
 ** the output, which is incorrect.
 */
-static void setJoinExpr(Expr *p, int iTable){
+SQLITE_PRIVATE void sqlite3SetJoinExpr(Expr *p, int iTable){
   while( p ){
     ExprSetProperty(p, EP_FromJoin);
     assert( !ExprHasProperty(p, EP_TokenOnly|EP_Reduced) );
@@ -125014,15 +128280,15 @@ static void setJoinExpr(Expr *p, int iTable){
     if( p->op==TK_FUNCTION && p->x.pList ){
       int i;
       for(i=0; i<p->x.pList->nExpr; i++){
-        setJoinExpr(p->x.pList->a[i].pExpr, iTable);
+        sqlite3SetJoinExpr(p->x.pList->a[i].pExpr, iTable);
       }
     }
-    setJoinExpr(p->pLeft, iTable);
+    sqlite3SetJoinExpr(p->pLeft, iTable);
     p = p->pRight;
   } 
 }
 
-/* Undo the work of setJoinExpr().  In the expression tree p, convert every
+/* Undo the work of sqlite3SetJoinExpr(). In the expression p, convert every
 ** term that is marked with EP_FromJoin and iRightJoinTable==iTable into
 ** an ordinary term that omits the EP_FromJoin mark.
 **
@@ -125089,10 +128355,11 @@ static int sqliteProcessJoin(Parse *pParse, Select *p){
         int iLeft;     /* Matching left table */
         int iLeftCol;  /* Matching column in the left table */
 
+        if( IsHiddenColumn(&pRightTab->aCol[j]) ) continue;
         zName = pRightTab->aCol[j].zName;
-        if( tableAndColumnIndex(pSrc, i+1, zName, &iLeft, &iLeftCol) ){
+        if( tableAndColumnIndex(pSrc, i+1, zName, &iLeft, &iLeftCol, 1) ){
           addWhereTerm(pParse, pSrc, iLeft, iLeftCol, i+1, j,
-                       isOuter, &p->pWhere);
+                isOuter, &p->pWhere);
         }
       }
     }
@@ -125109,8 +128376,8 @@ static int sqliteProcessJoin(Parse *pParse, Select *p){
     ** an AND operator.
     */
     if( pRight->pOn ){
-      if( isOuter ) setJoinExpr(pRight->pOn, pRight->iCursor);
-      p->pWhere = sqlite3ExprAnd(pParse->db, p->pWhere, pRight->pOn);
+      if( isOuter ) sqlite3SetJoinExpr(pRight->pOn, pRight->iCursor);
+      p->pWhere = sqlite3ExprAnd(pParse, p->pWhere, pRight->pOn);
       pRight->pOn = 0;
     }
 
@@ -125132,7 +128399,7 @@ static int sqliteProcessJoin(Parse *pParse, Select *p){
         zName = pList->a[j].zName;
         iRightCol = columnIndex(pRightTab, zName);
         if( iRightCol<0
-         || !tableAndColumnIndex(pSrc, i+1, zName, &iLeft, &iLeftCol)
+         || !tableAndColumnIndex(pSrc, i+1, zName, &iLeft, &iLeftCol, 0)
         ){
           sqlite3ErrorMsg(pParse, "cannot join using column %s - column "
             "not present in both tables", zName);
@@ -125284,11 +128551,12 @@ static void pushOntoSorter(
     if( pParse->db->mallocFailed ) return;
     pOp->p2 = nKey + nData;
     pKI = pOp->p4.pKeyInfo;
-    memset(pKI->aSortOrder, 0, pKI->nKeyField); /* Makes OP_Jump testable */
+    memset(pKI->aSortFlags, 0, pKI->nKeyField); /* Makes OP_Jump testable */
     sqlite3VdbeChangeP4(v, -1, (char*)pKI, P4_KEYINFO);
     testcase( pKI->nAllField > pKI->nKeyField+2 );
     pOp->p4.pKeyInfo = sqlite3KeyInfoFromExprList(pParse,pSort->pOrderBy,nOBSat,
                                            pKI->nAllField-pKI->nKeyField-1);
+    pOp = 0; /* Ensure pOp not used after sqltie3VdbeAddOp3() */
     addrJmp = sqlite3VdbeCurrentAddr(v);
     sqlite3VdbeAddOp3(v, OP_Jump, addrJmp+1, 0, addrJmp+1); VdbeCoverage(v);
     pSort->labelBkOut = sqlite3VdbeMakeLabel(pParse);
@@ -125537,7 +128805,7 @@ static void selectInnerLoop(
   if( srcTab>=0 ){
     for(i=0; i<nResultCol; i++){
       sqlite3VdbeAddOp3(v, OP_Column, srcTab, i, regResult+i);
-      VdbeComment((v, "%s", p->pEList->a[i].zName));
+      VdbeComment((v, "%s", p->pEList->a[i].zEName));
     }
   }else if( eDest!=SRT_Exists ){
 #ifdef SQLITE_ENABLE_SORTER_REFERENCES
@@ -125651,6 +128919,7 @@ static void selectInnerLoop(
         pOp->opcode = OP_Null;
         pOp->p1 = 1;
         pOp->p2 = regPrev;
+        pOp = 0;  /* Ensure pOp is not used after sqlite3VdbeAddOp() */
 
         iJump = sqlite3VdbeCurrentAddr(v) + nResultCol;
         for(i=0; i<nResultCol; i++){
@@ -125895,7 +129164,7 @@ SQLITE_PRIVATE KeyInfo *sqlite3KeyInfoAlloc(sqlite3 *db, int N, int X){
   int nExtra = (N+X)*(sizeof(CollSeq*)+1) - sizeof(CollSeq*);
   KeyInfo *p = sqlite3DbMallocRawNN(db, sizeof(KeyInfo) + nExtra);
   if( p ){
-    p->aSortOrder = (u8*)&p->aColl[N+X];
+    p->aSortFlags = (u8*)&p->aColl[N+X];
     p->nKeyField = (u16)N;
     p->nAllField = (u16)(N+X);
     p->enc = ENC(db);
@@ -125972,7 +129241,7 @@ SQLITE_PRIVATE KeyInfo *sqlite3KeyInfoFromExprList(
     assert( sqlite3KeyInfoIsWriteable(pInfo) );
     for(i=iStart, pItem=pList->a+iStart; i<nExpr; i++, pItem++){
       pInfo->aColl[i-iStart] = sqlite3ExprNNCollSeq(pParse, pItem->pExpr);
-      pInfo->aSortOrder[i-iStart] = pItem->sortOrder;
+      pInfo->aSortFlags[i-iStart] = pItem->sortFlags;
     }
   }
   return pInfo;
@@ -126157,7 +129426,7 @@ static void generateSortTail(
         iRead = iCol--;
       }
       sqlite3VdbeAddOp3(v, OP_Column, iSortTab, iRead, regRow+i);
-      VdbeComment((v, "%s", aOutEx[i].zName?aOutEx[i].zName : aOutEx[i].zSpan));
+      VdbeComment((v, "%s", aOutEx[i].zEName));
     }
   }
   switch( eDest ){
@@ -126264,8 +129533,6 @@ static const char *columnTypeImpl(
 
   assert( pExpr!=0 );
   assert( pNC->pSrcList!=0 );
-  assert( pExpr->op!=TK_AGG_COLUMN );  /* This routine runes before aggregates
-                                       ** are processed */
   switch( pExpr->op ){
     case TK_COLUMN: {
       /* The expression is a column. Locate the table the column is being
@@ -126493,9 +129760,9 @@ static void generateColumnNames(
     assert( p!=0 );
     assert( p->op!=TK_AGG_COLUMN );  /* Agg processing has not run yet */
     assert( p->op!=TK_COLUMN || p->y.pTab!=0 ); /* Covering idx not yet coded */
-    if( pEList->a[i].zName ){
+    if( pEList->a[i].zEName && pEList->a[i].eEName==ENAME_NAME ){
       /* An AS clause always takes first priority */
-      char *zName = pEList->a[i].zName;
+      char *zName = pEList->a[i].zEName;
       sqlite3VdbeSetColName(v, i, COLNAME_NAME, zName, SQLITE_TRANSIENT);
     }else if( srcName && p->op==TK_COLUMN ){
       char *zCol;
@@ -126517,7 +129784,7 @@ static void generateColumnNames(
         sqlite3VdbeSetColName(v, i, COLNAME_NAME, zCol, SQLITE_TRANSIENT);
       }
     }else{
-      const char *z = pEList->a[i].zSpan;
+      const char *z = pEList->a[i].zEName;
       z = z==0 ? sqlite3MPrintf(db, "column%d", i+1) : sqlite3DbStrDup(db, z);
       sqlite3VdbeSetColName(v, i, COLNAME_NAME, z, SQLITE_DYNAMIC);
     }
@@ -126579,15 +129846,14 @@ SQLITE_PRIVATE int sqlite3ColumnsFromExprList(
   for(i=0, pCol=aCol; i<nCol && !db->mallocFailed; i++, pCol++){
     /* Get an appropriate name for the column
     */
-    if( (zName = pEList->a[i].zName)!=0 ){
+    if( (zName = pEList->a[i].zEName)!=0 && pEList->a[i].eEName==ENAME_NAME ){
       /* If the column contains an "AS <name>" phrase, use <name> as the name */
     }else{
-      Expr *pColExpr = sqlite3ExprSkipCollate(pEList->a[i].pExpr);
+      Expr *pColExpr = sqlite3ExprSkipCollateAndLikely(pEList->a[i].pExpr);
       while( pColExpr->op==TK_DOT ){
         pColExpr = pColExpr->pRight;
         assert( pColExpr!=0 );
       }
-      assert( pColExpr->op!=TK_AGG_COLUMN );
       if( pColExpr->op==TK_COLUMN ){
         /* For columns use the column name name */
         int iCol = pColExpr->iColumn;
@@ -126600,10 +129866,10 @@ SQLITE_PRIVATE int sqlite3ColumnsFromExprList(
         zName = pColExpr->u.zToken;
       }else{
         /* Use the original text of the column expression as its name */
-        zName = pEList->a[i].zSpan;
+        zName = pEList->a[i].zEName;
       }
     }
-    if( zName ){
+    if( zName && !sqlite3IsTrueOrFalse(zName) ){
       zName = sqlite3DbStrDup(db, zName);
     }else{
       zName = sqlite3MPrintf(db,"column%d",i+1);
@@ -126655,7 +129921,8 @@ SQLITE_PRIVATE int sqlite3ColumnsFromExprList(
 SQLITE_PRIVATE void sqlite3SelectAddColumnTypeAndCollation(
   Parse *pParse,        /* Parsing contexts */
   Table *pTab,          /* Add column type information to this table */
-  Select *pSelect       /* SELECT used to determine types and collations */
+  Select *pSelect,      /* SELECT used to determine types and collations */
+  char aff              /* Default affinity for columns */
 ){
   sqlite3 *db = pParse->db;
   NameContext sNC;
@@ -126688,7 +129955,7 @@ SQLITE_PRIVATE void sqlite3SelectAddColumnTypeAndCollation(
         pCol->colFlags |= COLFLAG_HASTYPE;
       }
     }
-    if( pCol->affinity==0 ) pCol->affinity = SQLITE_AFF_BLOB;
+    if( pCol->affinity<=SQLITE_AFF_NONE ) pCol->affinity = aff;
     pColl = sqlite3ExprCollSeq(pParse, p);
     if( pColl && pCol->zColl==0 ){
       pCol->zColl = sqlite3DbStrDup(db, pColl->zName);
@@ -126701,7 +129968,7 @@ SQLITE_PRIVATE void sqlite3SelectAddColumnTypeAndCollation(
 ** Given a SELECT statement, generate a Table structure that describes
 ** the result set of that SELECT.
 */
-SQLITE_PRIVATE Table *sqlite3ResultSetOfSelect(Parse *pParse, Select *pSelect){
+SQLITE_PRIVATE Table *sqlite3ResultSetOfSelect(Parse *pParse, Select *pSelect, char aff){
   Table *pTab;
   sqlite3 *db = pParse->db;
   u64 savedFlags;
@@ -126717,14 +129984,11 @@ SQLITE_PRIVATE Table *sqlite3ResultSetOfSelect(Parse *pParse, Select *pSelect){
   if( pTab==0 ){
     return 0;
   }
-  /* The sqlite3ResultSetOfSelect() is only used n contexts where lookaside
-  ** is disabled */
-  assert( db->lookaside.bDisable );
   pTab->nTabRef = 1;
   pTab->zName = 0;
   pTab->nRowLogEst = 200; assert( 200==sqlite3LogEst(1048576) );
   sqlite3ColumnsFromExprList(pParse, pSelect->pEList, &pTab->nCol, &pTab->aCol);
-  sqlite3SelectAddColumnTypeAndCollation(pParse, pTab, pSelect);
+  sqlite3SelectAddColumnTypeAndCollation(pParse, pTab, pSelect, aff);
   pTab->iPKey = -1;
   if( db->mallocFailed ){
     sqlite3DeleteTable(db, pTab);
@@ -126878,7 +130142,7 @@ static KeyInfo *multiSelectOrderByKeyInfo(Parse *pParse, Select *p, int nExtra){
       }
       assert( sqlite3KeyInfoIsWriteable(pRet) );
       pRet->aColl[i] = pColl;
-      pRet->aSortOrder[i] = pOrderBy->a[i].sortOrder;
+      pRet->aSortFlags[i] = pOrderBy->a[i].sortFlags;
     }
   }
 
@@ -127097,6 +130361,9 @@ static int multiSelectValues(
     assert( p->selFlags & SF_Values );
     assert( p->op==TK_ALL || (p->op==TK_SELECT && p->pPrior==0) );
     assert( p->pNext==0 || p->pEList->nExpr==p->pNext->pEList->nExpr );
+#ifndef SQLITE_OMIT_WINDOWFUNC
+    if( p->pWin ) return -1;
+#endif
     if( p->pPrior==0 ) break;
     assert( p->pPrior->pNext==p );
     p = p->pPrior;
@@ -127161,6 +130428,7 @@ static int multiSelect(
   */
   assert( p && p->pPrior );  /* Calling function guarantees this much */
   assert( (p->selFlags & SF_Recursive)==0 || p->op==TK_ALL || p->op==TK_UNION );
+  assert( p->selFlags & SF_Compound );
   db = pParse->db;
   pPrior = p->pPrior;
   dest = *pDest;
@@ -127186,7 +130454,8 @@ static int multiSelect(
   */
   if( p->selFlags & SF_MultiValue ){
     rc = multiSelectValues(pParse, p, &dest);
-    goto multi_select_end;
+    if( rc>=0 ) goto multi_select_end;
+    rc = SQLITE_OK;
   }
 
   /* Make sure all SELECTs in the statement have the same number of elements
@@ -127331,9 +130600,9 @@ static int multiSelect(
         ** it is that we currently need.
         */
         assert( unionTab==dest.iSDParm || dest.eDest!=priorOp );
-        if( dest.eDest!=priorOp ){
+        assert( p->pEList || db->mallocFailed );
+        if( dest.eDest!=priorOp && db->mallocFailed==0 ){
           int iCont, iBreak, iStart;
-          assert( p->pEList );
           iBreak = sqlite3VdbeMakeLabel(pParse);
           iCont = sqlite3VdbeMakeLabel(pParse);
           computeLimitRegisters(pParse, p, iBreak);
@@ -127429,6 +130698,7 @@ static int multiSelect(
     }
   #endif
   }
+  if( pParse->nErr ) goto multi_select_end;
   
   /* Compute collating sequences used by 
   ** temporary tables needed to implement the compound select.
@@ -127588,11 +130858,14 @@ static int generateOutputSubroutine(
 
     /* If this is a scalar select that is part of an expression, then
     ** store the results in the appropriate memory cell and break out
-    ** of the scan loop.
+    ** of the scan loop.  Note that the select might return multiple columns
+    ** if it is the RHS of a row-value IN operator.
     */
     case SRT_Mem: {
-      assert( pIn->nSdst==1 || pParse->nErr>0 );  testcase( pIn->nSdst!=1 );
-      sqlite3ExprCodeMove(pParse, pIn->iSdst, pDest->iSDParm, 1);
+      if( pParse->nErr==0 ){
+        testcase( pIn->nSdst>1 );
+        sqlite3ExprCodeMove(pParse, pIn->iSdst, pDest->iSDParm, pIn->nSdst);
+      }
       /* The LIMIT clause will jump out of the loop for us */
       break;
     }
@@ -127849,7 +131122,7 @@ static int multiSelectOrderBy(
       assert( sqlite3KeyInfoIsWriteable(pKeyDup) );
       for(i=0; i<nExpr; i++){
         pKeyDup->aColl[i] = multiSelectCollSeq(pParse, p, i);
-        pKeyDup->aSortOrder[i] = 0;
+        pKeyDup->aSortFlags[i] = 0;
       }
     }
   }
@@ -128099,6 +131372,18 @@ static Expr *substExpr(
         }
         sqlite3ExprDelete(db, pExpr);
         pExpr = pNew;
+
+        /* Ensure that the expression now has an implicit collation sequence,
+        ** just as it did when it was a column of a view or sub-query. */
+        if( pExpr ){
+          if( pExpr->op!=TK_COLUMN && pExpr->op!=TK_COLLATE ){
+            CollSeq *pColl = sqlite3ExprCollSeq(pSubst->pParse, pExpr);
+            pExpr = sqlite3ExprAddCollateString(pSubst->pParse, pExpr, 
+                (pColl ? pColl->zName : "BINARY")
+            );
+          }
+          ExprClearProperty(pExpr, EP_Collate);
+        }
       }
     }
   }else{
@@ -128112,6 +131397,14 @@ static Expr *substExpr(
     }else{
       substExprList(pSubst, pExpr->x.pList);
     }
+#ifndef SQLITE_OMIT_WINDOWFUNC
+    if( ExprHasProperty(pExpr, EP_WinFunc) ){
+      Window *pWin = pExpr->y.pWin;
+      pWin->pFilter = substExpr(pSubst, pWin->pFilter);
+      substExprList(pSubst, pWin->pPartition);
+      substExprList(pSubst, pWin->pOrderBy);
+    }
+#endif
   }
   return pExpr;
 }
@@ -128197,6 +131490,7 @@ static void substSelect(
 **        (3b) the FROM clause of the subquery may not contain a virtual
 **             table and
 **        (3c) the outer query may not be an aggregate.
+**        (3d) the outer query may not be DISTINCT.
 **
 **   (4)  The subquery can not be DISTINCT.
 **
@@ -128247,6 +131541,7 @@ static void substSelect(
 **              (17d1) aggregate, or
 **              (17d2) DISTINCT, or
 **              (17d3) a join.
+**        (17e) the subquery may not contain window functions
 **
 **        The parent and sub-query may contain WHERE clauses. Subject to
 **        rules (11), (13) and (14), they may also contain ORDER BY,
@@ -128393,8 +131688,11 @@ static int flattenSubquery(
   */
   if( (pSubitem->fg.jointype & JT_OUTER)!=0 ){
     isLeftJoin = 1;
-    if( pSubSrc->nSrc>1 || isAgg || IsVirtual(pSubSrc->a[0].pTab) ){
-      /*  (3a)             (3c)     (3b) */
+    if( pSubSrc->nSrc>1                   /* (3a) */
+     || isAgg                             /* (3b) */
+     || IsVirtual(pSubSrc->a[0].pTab)     /* (3c) */
+     || (p->selFlags & SF_Distinct)!=0    /* (3d) */
+    ){
       return 0;
     }
   }
@@ -128428,6 +131726,9 @@ static int flattenSubquery(
       if( (pSub1->selFlags & (SF_Distinct|SF_Aggregate))!=0    /* (17b) */
        || (pSub1->pPrior && pSub1->op!=TK_ALL)                 /* (17a) */
        || pSub1->pSrc->nSrc<1                                  /* (17c) */
+#ifndef SQLITE_OMIT_WINDOWFUNC
+       || pSub1->pWin                                          /* (17e) */
+#endif
       ){
         return 0;
       }
@@ -128572,6 +131873,7 @@ static int flattenSubquery(
   for(pParent=p; pParent; pParent=pParent->pPrior, pSub=pSub->pPrior){
     int nSubSrc;
     u8 jointype = 0;
+    assert( pSub!=0 );
     pSubSrc = pSub->pSrc;     /* FROM clause of subquery */
     nSubSrc = pSubSrc->nSrc;  /* Number of terms in subquery FROM clause */
     pSrc = pParent->pSrc;     /* FROM clause of the outer query */
@@ -128653,9 +131955,9 @@ static int flattenSubquery(
     pWhere = pSub->pWhere;
     pSub->pWhere = 0;
     if( isLeftJoin>0 ){
-      setJoinExpr(pWhere, iNewParent);
+      sqlite3SetJoinExpr(pWhere, iNewParent);
     }
-    pParent->pWhere = sqlite3ExprAnd(db, pWhere, pParent->pWhere);
+    pParent->pWhere = sqlite3ExprAnd(pParse, pWhere, pParent->pWhere);
     if( db->mallocFailed==0 ){
       SubstContext x;
       x.pParse = pParse;
@@ -128666,10 +131968,10 @@ static int flattenSubquery(
       substSelect(&x, pParent, 0);
     }
   
-    /* The flattened query is distinct if either the inner or the
-    ** outer query is distinct. 
-    */
-    pParent->selFlags |= pSub->selFlags & SF_Distinct;
+    /* The flattened query is a compound if either the inner or the
+    ** outer query is a compound. */
+    pParent->selFlags |= pSub->selFlags & SF_Compound;
+    assert( (pSub->selFlags & SF_Distinct)==0 ); /* restriction (17b) */
   
     /*
     ** SELECT ... FROM (SELECT ... LIMIT a OFFSET b) LIMIT x OFFSET y;
@@ -128713,23 +132015,36 @@ struct WhereConst {
 
 /*
 ** Add a new entry to the pConst object.  Except, do not add duplicate
-** pColumn entires.
+** pColumn entires.  Also, do not add if doing so would not be appropriate.
+**
+** The caller guarantees the pColumn is a column and pValue is a constant.
+** This routine has to do some additional checks before completing the
+** insert.
 */
 static void constInsert(
-  WhereConst *pConst,      /* The WhereConst into which we are inserting */
-  Expr *pColumn,           /* The COLUMN part of the constraint */
-  Expr *pValue             /* The VALUE part of the constraint */
+  WhereConst *pConst,  /* The WhereConst into which we are inserting */
+  Expr *pColumn,       /* The COLUMN part of the constraint */
+  Expr *pValue,        /* The VALUE part of the constraint */
+  Expr *pExpr          /* Overall expression: COLUMN=VALUE or VALUE=COLUMN */
 ){
   int i;
   assert( pColumn->op==TK_COLUMN );
+  assert( sqlite3ExprIsConstant(pValue) );
+
+  if( !ExprHasProperty(pValue, EP_FixedCol) && sqlite3ExprAffinity(pValue)!=0 ){
+    return;
+  }
+  if( !sqlite3IsBinary(sqlite3ExprCompareCollSeq(pConst->pParse,pExpr)) ){
+    return;
+  }
 
   /* 2018-10-25 ticket [cf5ed20f]
   ** Make sure the same pColumn is not inserted more than once */
   for(i=0; i<pConst->nConst; i++){
-    const Expr *pExpr = pConst->apExpr[i*2];
-    assert( pExpr->op==TK_COLUMN );
-    if( pExpr->iTable==pColumn->iTable
-     && pExpr->iColumn==pColumn->iColumn
+    const Expr *pE2 = pConst->apExpr[i*2];
+    assert( pE2->op==TK_COLUMN );
+    if( pE2->iTable==pColumn->iTable
+     && pE2->iColumn==pColumn->iColumn
     ){
       return;  /* Already present.  Return without doing anything. */
     }
@@ -128741,7 +132056,9 @@ static void constInsert(
   if( pConst->apExpr==0 ){
     pConst->nConst = 0;
   }else{
-    if( ExprHasProperty(pValue, EP_FixedCol) ) pValue = pValue->pLeft;
+    if( ExprHasProperty(pValue, EP_FixedCol) ){
+      pValue = pValue->pLeft;
+    }
     pConst->apExpr[pConst->nConst*2-2] = pColumn;
     pConst->apExpr[pConst->nConst*2-1] = pValue;
   }
@@ -128767,19 +132084,11 @@ static void findConstInWhere(WhereConst *pConst, Expr *pExpr){
   pLeft = pExpr->pLeft;
   assert( pRight!=0 );
   assert( pLeft!=0 );
-  if( pRight->op==TK_COLUMN
-   && !ExprHasProperty(pRight, EP_FixedCol)
-   && sqlite3ExprIsConstant(pLeft)
-   && sqlite3IsBinary(sqlite3BinaryCompareCollSeq(pConst->pParse,pLeft,pRight))
-  ){
-    constInsert(pConst, pRight, pLeft);
-  }else
-  if( pLeft->op==TK_COLUMN
-   && !ExprHasProperty(pLeft, EP_FixedCol)
-   && sqlite3ExprIsConstant(pRight)
-   && sqlite3IsBinary(sqlite3BinaryCompareCollSeq(pConst->pParse,pLeft,pRight))
-  ){
-    constInsert(pConst, pLeft, pRight);
+  if( pRight->op==TK_COLUMN && sqlite3ExprIsConstant(pLeft) ){
+    constInsert(pConst,pRight,pLeft,pExpr);
+  }
+  if( pLeft->op==TK_COLUMN && sqlite3ExprIsConstant(pRight) ){
+    constInsert(pConst,pLeft,pRight,pExpr);
   }
 }
 
@@ -128793,7 +132102,11 @@ static int propagateConstantExprRewrite(Walker *pWalker, Expr *pExpr){
   int i;
   WhereConst *pConst;
   if( pExpr->op!=TK_COLUMN ) return WRC_Continue;
-  if( ExprHasProperty(pExpr, EP_FixedCol) ) return WRC_Continue;
+  if( ExprHasProperty(pExpr, EP_FixedCol|EP_FromJoin) ){
+    testcase( ExprHasProperty(pExpr, EP_FixedCol) );
+    testcase( ExprHasProperty(pExpr, EP_FromJoin) );
+    return WRC_Continue;
+  }
   pConst = pWalker->u.pConst;
   for(i=0; i<pConst->nConst; i++){
     Expr *pColumn = pConst->apExpr[i*2];
@@ -128815,10 +132128,9 @@ static int propagateConstantExprRewrite(Walker *pWalker, Expr *pExpr){
 ** The WHERE-clause constant propagation optimization.
 **
 ** If the WHERE clause contains terms of the form COLUMN=CONSTANT or
-** CONSTANT=COLUMN that must be tree (in other words, if the terms top-level
-** AND-connected terms that are not part of a ON clause from a LEFT JOIN)
-** then throughout the query replace all other occurrences of COLUMN
-** with CONSTANT within the WHERE clause.
+** CONSTANT=COLUMN that are top-level AND-connected terms that are not
+** part of a ON clause from a LEFT JOIN, then throughout the query
+** replace all other occurrences of COLUMN with CONSTANT.
 **
 ** For example, the query:
 **
@@ -128990,9 +132302,9 @@ static int pushDownWhereTerms(
       x.pEList = pSubq->pEList;
       pNew = substExpr(&x, pNew);
       if( pSubq->selFlags & SF_Aggregate ){
-        pSubq->pHaving = sqlite3ExprAnd(pParse->db, pSubq->pHaving, pNew);
+        pSubq->pHaving = sqlite3ExprAnd(pParse, pSubq->pHaving, pNew);
       }else{
-        pSubq->pWhere = sqlite3ExprAnd(pParse->db, pSubq->pWhere, pNew);
+        pSubq->pWhere = sqlite3ExprAnd(pParse, pSubq->pWhere, pNew);
       }
       pSubq = pSubq->pPrior;
     }
@@ -129022,24 +132334,27 @@ static u8 minMaxQuery(sqlite3 *db, Expr *pFunc, ExprList **ppMinMax){
   ExprList *pEList = pFunc->x.pList;    /* Arguments to agg function */
   const char *zFunc;                    /* Name of aggregate function pFunc */
   ExprList *pOrderBy;
-  u8 sortOrder;
+  u8 sortFlags;
 
   assert( *ppMinMax==0 );
   assert( pFunc->op==TK_AGG_FUNCTION );
-  if( pEList==0 || pEList->nExpr!=1 ) return eRet;
+  assert( !IsWindowFunc(pFunc) );
+  if( pEList==0 || pEList->nExpr!=1 || ExprHasProperty(pFunc, EP_WinFunc) ){
+    return eRet;
+  }
   zFunc = pFunc->u.zToken;
   if( sqlite3StrICmp(zFunc, "min")==0 ){
     eRet = WHERE_ORDERBY_MIN;
-    sortOrder = SQLITE_SO_ASC;
+    sortFlags = KEYINFO_ORDER_BIGNULL;
   }else if( sqlite3StrICmp(zFunc, "max")==0 ){
     eRet = WHERE_ORDERBY_MAX;
-    sortOrder = SQLITE_SO_DESC;
+    sortFlags = KEYINFO_ORDER_DESC;
   }else{
     return eRet;
   }
   *ppMinMax = pOrderBy = sqlite3ExprListDup(db, pEList, 0);
   assert( pOrderBy!=0 || db->mallocFailed );
-  if( pOrderBy ) pOrderBy->a[0].sortOrder = sortOrder;
+  if( pOrderBy ) pOrderBy->a[0].sortFlags = sortFlags;
   return eRet;
 }
 
@@ -129073,7 +132388,7 @@ static Table *isSimpleCount(Select *p, AggInfo *pAggInfo){
   if( pExpr->op!=TK_AGG_FUNCTION ) return 0;
   if( NEVER(pAggInfo->nFunc==0) ) return 0;
   if( (pAggInfo->aFunc[0].pFunc->funcFlags&SQLITE_FUNC_COUNT)==0 ) return 0;
-  if( pExpr->flags&EP_Distinct ) return 0;
+  if( ExprHasProperty(pExpr, EP_Distinct|EP_WinFunc) ) return 0;
 
   return pTab;
 }
@@ -129164,6 +132479,9 @@ static int convertCompoundSelectToSubquery(Walker *pWalker, Select *p){
   p->pPrior = 0;
   p->pNext = 0;
   p->pWith = 0;
+#ifndef SQLITE_OMIT_WINDOWFUNC
+  p->pWinDefn = 0;
+#endif
   p->selFlags &= ~SF_Compound;
   assert( (p->selFlags & SF_Converted)==0 );
   p->selFlags |= SF_Converted;
@@ -129263,6 +132581,9 @@ static int withExpand(
   With *pWith;                    /* WITH clause that pCte belongs to */
 
   assert( pFrom->pTab==0 );
+  if( pParse->nErr ){
+    return SQLITE_ERROR;
+  }
 
   pCte = searchWith(pParse->pWith, pFrom, &pWith);
   if( pCte ){
@@ -129383,7 +132704,7 @@ static void selectPopWith(Walker *pWalker, Select *p){
   if( OK_IF_ALWAYS_TRUE(pParse->pWith) && p->pPrior==0 ){
     With *pWith = findRightmost(p)->pWith;
     if( pWith!=0 ){
-      assert( pParse->pWith==pWith );
+      assert( pParse->pWith==pWith || pParse->nErr );
       pParse->pWith = pWith->pOuter;
     }
   }
@@ -129418,7 +132739,7 @@ SQLITE_PRIVATE int sqlite3ExpandSubquery(Parse *pParse, struct SrcList_item *pFr
   pTab->nRowLogEst = 200; assert( 200==sqlite3LogEst(1048576) );
   pTab->tabFlags |= TF_Ephemeral;
 
-  return SQLITE_OK;
+  return pParse->nErr ? SQLITE_ERROR : SQLITE_OK;
 }
 
 /*
@@ -129464,6 +132785,10 @@ static int selectExpander(Walker *pWalker, Select *p){
   if( (selFlags & SF_Expanded)!=0 ){
     return WRC_Prune;
   }
+  if( pWalker->eCode ){
+    /* Renumber selId because it has been copied from a view */
+    p->selId = ++pParse->nSelect;
+  }
   pTabList = p->pSrc;
   pEList = p->pEList;
   sqlite3WithPush(pParse, p->pWith, 0);
@@ -129510,15 +132835,32 @@ static int selectExpander(Walker *pWalker, Select *p){
       if( !IsVirtual(pTab) && cannotBeFunction(pParse, pFrom) ){
         return WRC_Abort;
       }
-#if !defined(SQLITE_OMIT_VIEW) || !defined (SQLITE_OMIT_VIRTUALTABLE)
+#if !defined(SQLITE_OMIT_VIEW) || !defined(SQLITE_OMIT_VIRTUALTABLE)
       if( IsVirtual(pTab) || pTab->pSelect ){
         i16 nCol;
+        u8 eCodeOrig = pWalker->eCode;
         if( sqlite3ViewGetColumnNames(pParse, pTab) ) return WRC_Abort;
         assert( pFrom->pSelect==0 );
+        if( pTab->pSelect && (db->flags & SQLITE_EnableView)==0 ){
+          sqlite3ErrorMsg(pParse, "access to view \"%s\" prohibited",
+            pTab->zName);
+        }
+#ifndef SQLITE_OMIT_VIRTUALTABLE
+        if( IsVirtual(pTab)
+         && pFrom->fg.fromDDL
+         && ALWAYS(pTab->pVTable!=0)
+         && pTab->pVTable->eVtabRisk > ((db->flags & SQLITE_TrustedSchema)!=0)
+        ){
+          sqlite3ErrorMsg(pParse, "unsafe use of virtual table \"%s\"",
+                                  pTab->zName);
+        }
+#endif
         pFrom->pSelect = sqlite3SelectDup(db, pTab->pSelect, 0);
         nCol = pTab->nCol;
         pTab->nCol = -1;
+        pWalker->eCode = 1;  /* Turn on Select.selId renumbering */
         sqlite3WalkSelect(pWalker, pFrom->pSelect);
+        pWalker->eCode = eCodeOrig;
         pTab->nCol = nCol;
       }
 #endif
@@ -129532,7 +132874,7 @@ static int selectExpander(Walker *pWalker, Select *p){
 
   /* Process NATURAL keywords, and ON and USING clauses of joins.
   */
-  if( db->mallocFailed || sqliteProcessJoin(pParse, p) ){
+  if( pParse->nErr || db->mallocFailed || sqliteProcessJoin(pParse, p) ){
     return WRC_Abort;
   }
 
@@ -129579,10 +132921,9 @@ static int selectExpander(Walker *pWalker, Select *p){
         */
         pNew = sqlite3ExprListAppend(pParse, pNew, a[k].pExpr);
         if( pNew ){
-          pNew->a[pNew->nExpr-1].zName = a[k].zName;
-          pNew->a[pNew->nExpr-1].zSpan = a[k].zSpan;
-          a[k].zName = 0;
-          a[k].zSpan = 0;
+          pNew->a[pNew->nExpr-1].zEName = a[k].zEName;
+          pNew->a[pNew->nExpr-1].eEName = a[k].eEName;
+          a[k].zEName = 0;
         }
         a[k].pExpr = 0;
       }else{
@@ -129621,7 +132962,7 @@ static int selectExpander(Walker *pWalker, Select *p){
 
             assert( zName );
             if( zTName && pSub
-             && sqlite3MatchSpanName(pSub->pEList->a[j].zSpan, 0, zTName, 0)==0
+             && sqlite3MatchEName(&pSub->pEList->a[j], 0, zTName, 0)==0
             ){
               continue;
             }
@@ -129639,7 +132980,7 @@ static int selectExpander(Walker *pWalker, Select *p){
 
             if( i>0 && zTName==0 ){
               if( (pFrom->fg.jointype & JT_NATURAL)!=0
-                && tableAndColumnIndex(pTabList, i, zName, 0, 0)
+                && tableAndColumnIndex(pTabList, i, zName, 0, 0, 1)
               ){
                 /* In a NATURAL join, omit the join columns from the 
                 ** table to the right of the join */
@@ -129674,15 +133015,16 @@ static int selectExpander(Walker *pWalker, Select *p){
             sqlite3ExprListSetName(pParse, pNew, &sColname, 0);
             if( pNew && (p->selFlags & SF_NestedFrom)!=0 ){
               struct ExprList_item *pX = &pNew->a[pNew->nExpr-1];
+              sqlite3DbFree(db, pX->zEName);
               if( pSub ){
-                pX->zSpan = sqlite3DbStrDup(db, pSub->pEList->a[j].zSpan);
-                testcase( pX->zSpan==0 );
+                pX->zEName = sqlite3DbStrDup(db, pSub->pEList->a[j].zEName);
+                testcase( pX->zEName==0 );
               }else{
-                pX->zSpan = sqlite3MPrintf(db, "%s.%s.%s",
+                pX->zEName = sqlite3MPrintf(db, "%s.%s.%s",
                                            zSchemaName, zTabName, zColname);
-                testcase( pX->zSpan==0 );
+                testcase( pX->zEName==0 );
               }
-              pX->bSpanIsTab = 1;
+              pX->eEName = ENAME_TAB;
             }
             sqlite3DbFree(db, zToFree);
           }
@@ -129768,6 +133110,7 @@ static void sqlite3SelectExpand(Parse *pParse, Select *pSelect){
   }
   w.xSelectCallback = selectExpander;
   w.xSelectCallback2 = selectPopWith;
+  w.eCode = 0;
   sqlite3WalkSelect(&w, pSelect);
 }
 
@@ -129805,7 +133148,8 @@ static void selectAddSubqueryTypeInfo(Walker *pWalker, Select *p){
       Select *pSel = pFrom->pSelect;
       if( pSel ){
         while( pSel->pPrior ) pSel = pSel->pPrior;
-        sqlite3SelectAddColumnTypeAndCollation(pParse, pTab, pSel);
+        sqlite3SelectAddColumnTypeAndCollation(pParse, pTab, pSel,
+                                               SQLITE_AFF_NONE);
       }
     }
   }
@@ -129945,6 +133289,25 @@ static void updateAccumulator(Parse *pParse, int regAcc, AggInfo *pAggInfo){
     int regAgg;
     ExprList *pList = pF->pExpr->x.pList;
     assert( !ExprHasProperty(pF->pExpr, EP_xIsSelect) );
+    assert( !IsWindowFunc(pF->pExpr) );
+    if( ExprHasProperty(pF->pExpr, EP_WinFunc) ){
+      Expr *pFilter = pF->pExpr->y.pWin->pFilter;
+      if( pAggInfo->nAccumulator 
+       && (pF->pFunc->funcFlags & SQLITE_FUNC_NEEDCOLL) 
+      ){
+        if( regHit==0 ) regHit = ++pParse->nMem;
+        /* If this is the first row of the group (regAcc==0), clear the
+        ** "magnet" register regHit so that the accumulator registers
+        ** are populated if the FILTER clause jumps over the the 
+        ** invocation of min() or max() altogether. Or, if this is not
+        ** the first row (regAcc==1), set the magnet register so that the
+        ** accumulators are not populated unless the min()/max() is invoked and
+        ** indicates that they should be.  */
+        sqlite3VdbeAddOp2(v, OP_Copy, regAcc, regHit);
+      }
+      addrNext = sqlite3VdbeMakeLabel(pParse);
+      sqlite3ExprIfFalse(pParse, pFilter, addrNext, SQLITE_JUMPIFNULL);
+    }
     if( pList ){
       nArg = pList->nExpr;
       regAgg = sqlite3GetTempRange(pParse, nArg);
@@ -129954,7 +133317,9 @@ static void updateAccumulator(Parse *pParse, int regAcc, AggInfo *pAggInfo){
       regAgg = 0;
     }
     if( pF->iDistinct>=0 ){
-      addrNext = sqlite3VdbeMakeLabel(pParse);
+      if( addrNext==0 ){ 
+        addrNext = sqlite3VdbeMakeLabel(pParse);
+      }
       testcase( nArg==0 );  /* Error condition */
       testcase( nArg>1 );   /* Also an error */
       codeDistinct(pParse, pF->iDistinct, addrNext, 1, regAgg);
@@ -129990,6 +133355,7 @@ static void updateAccumulator(Parse *pParse, int regAcc, AggInfo *pAggInfo){
   for(i=0, pC=pAggInfo->aCol; i<pAggInfo->nAccumulator; i++, pC++){
     sqlite3ExprCode(pParse, pC->pExpr, pC->iMem);
   }
+
   pAggInfo->directMode = 0;
   if( addrHitTest ){
     sqlite3VdbeJumpHere(v, addrHitTest);
@@ -130035,11 +133401,11 @@ static int havingToWhereExprCb(Walker *pWalker, Expr *pExpr){
     Select *pS = pWalker->u.pSelect;
     if( sqlite3ExprIsConstantOrGroupBy(pWalker->pParse, pExpr, pS->pGroupBy) ){
       sqlite3 *db = pWalker->pParse->db;
-      Expr *pNew = sqlite3ExprAlloc(db, TK_INTEGER, &sqlite3IntTokens[1], 0);
+      Expr *pNew = sqlite3Expr(db, TK_INTEGER, "1");
       if( pNew ){
         Expr *pWhere = pS->pWhere;
         SWAP(Expr, *pNew, *pExpr);
-        pNew = sqlite3ExprAnd(db, pWhere, pNew);
+        pNew = sqlite3ExprAnd(pWalker->pParse, pWhere, pNew);
         pS->pWhere = pNew;
         pWalker->eCode = 1;
       }
@@ -130094,15 +133460,19 @@ static struct SrcList_item *isSelfJoinView(
     if( pItem->pSelect==0 ) continue;
     if( pItem->fg.viaCoroutine ) continue;
     if( pItem->zName==0 ) continue;
-    if( sqlite3_stricmp(pItem->zDatabase, pThis->zDatabase)!=0 ) continue;
+    assert( pItem->pTab!=0 );
+    assert( pThis->pTab!=0 );
+    if( pItem->pTab->pSchema!=pThis->pTab->pSchema ) continue;
     if( sqlite3_stricmp(pItem->zName, pThis->zName)!=0 ) continue;
     pS1 = pItem->pSelect;
-    if( pThis->pSelect->selId!=pS1->selId ){
+    if( pItem->pTab->pSchema==0 && pThis->pSelect->selId!=pS1->selId ){
       /* The query flattener left two different CTE tables with identical
       ** names in the same FROM clause. */
       continue;
     }
-    if( sqlite3ExprCompare(0, pThis->pSelect->pWhere, pS1->pWhere, -1) ){
+    if( sqlite3ExprCompare(0, pThis->pSelect->pWhere, pS1->pWhere, -1)
+     || sqlite3ExprCompare(0, pThis->pSelect->pHaving, pS1->pHaving, -1) 
+    ){
       /* The view was modified by some other optimization such as
       ** pushDownWhereTerms() */
       continue;
@@ -130127,7 +133497,8 @@ static struct SrcList_item *isSelfJoinView(
 **   *  The subquery is a UNION ALL of two or more terms
 **   *  The subquery does not have a LIMIT clause
 **   *  There is no WHERE or GROUP BY or HAVING clauses on the subqueries
-**   *  The outer query is a simple count(*)
+**   *  The outer query is a simple count(*) with no WHERE clause or other
+**      extraneous syntax.
 **
 ** Return TRUE if the optimization is undertaken.
 */
@@ -130138,6 +133509,8 @@ static int countOfViewOptimization(Parse *pParse, Select *p){
   sqlite3 *db;
   if( (p->selFlags & SF_Aggregate)==0 ) return 0;   /* This is an aggregate */
   if( p->pEList->nExpr!=1 ) return 0;               /* Single result column */
+  if( p->pWhere ) return 0;
+  if( p->pGroupBy ) return 0;
   pExpr = p->pEList->a[0].pExpr;
   if( pExpr->op!=TK_AGG_FUNCTION ) return 0;        /* Result is an aggregate */
   if( sqlite3_stricmp(pExpr->u.zToken,"count") ) return 0;  /* Is count() */
@@ -130278,11 +133651,13 @@ SQLITE_PRIVATE int sqlite3Select(
   }
 
 #ifndef SQLITE_OMIT_WINDOWFUNC
-  if( sqlite3WindowRewrite(pParse, p) ){
+  rc = sqlite3WindowRewrite(pParse, p);
+  if( rc ){
+    assert( db->mallocFailed || pParse->nErr>0 );
     goto select_end;
   }
 #if SELECTTRACE_ENABLED
-  if( sqlite3SelectTrace & 0x108 ){
+  if( p->pWin && (sqlite3SelectTrace & 0x108)!=0 ){
     SELECTTRACE(0x104,pParse,p, ("after window rewrite:\n"));
     sqlite3TreeViewSelect(0, p, 0);
   }
@@ -130450,7 +133825,7 @@ SQLITE_PRIVATE int sqlite3Select(
     ** assume the column name is non-NULL and segfault.  The use of an empty
     ** string for the fake column name seems safer.
     */
-    if( pItem->colUsed==0 ){
+    if( pItem->colUsed==0 && pItem->zName!=0 ){
       sqlite3AuthCheck(pParse, SQLITE_READ, pItem->zName, "", pItem->zDatabase);
     }
 
@@ -130464,8 +133839,15 @@ SQLITE_PRIVATE int sqlite3Select(
     ** technically harmless for it to be generated multiple times. The
     ** following assert() will detect if something changes to cause
     ** the same subquery to be coded multiple times, as a signal to the
-    ** developers to try to optimize the situation. */
-    assert( pItem->addrFillSub==0 );
+    ** developers to try to optimize the situation.
+    **
+    ** Update 2019-07-24:
+    ** See ticket https://sqlite.org/src/tktview/c52b09c7f38903b1311cec40.
+    ** The dbsqlfuzz fuzzer found a case where the same subquery gets
+    ** coded twice.  So this assert() now becomes a testcase().  It should
+    ** be very rare, though.
+    */
+    testcase( pItem->addrFillSub!=0 );
 
     /* Increment Parse.nHeight by the height of the largest expression
     ** tree referred to by this, the parent select. The child select
@@ -130539,7 +133921,7 @@ SQLITE_PRIVATE int sqlite3Select(
       int retAddr;
       struct SrcList_item *pPrior;
 
-      assert( pItem->addrFillSub==0 );
+      testcase( pItem->addrFillSub==0 ); /* Ticket c52b09c7f38903b1311 */
       pItem->regReturn = ++pParse->nMem;
       topAddr = sqlite3VdbeAddOp2(v, OP_Integer, 0, pItem->regReturn);
       pItem->addrFillSub = topAddr+1;
@@ -130607,9 +133989,13 @@ SQLITE_PRIVATE int sqlite3Select(
   */
   if( (p->selFlags & (SF_Distinct|SF_Aggregate))==SF_Distinct 
    && sqlite3ExprListCompare(sSort.pOrderBy, pEList, -1)==0
+#ifndef SQLITE_OMIT_WINDOWFUNC
+   && p->pWin==0
+#endif
   ){
     p->selFlags &= ~SF_Distinct;
     pGroupBy = p->pGroupBy = sqlite3ExprListDup(db, pEList, 0);
+    p->selFlags |= SF_Aggregate;
     /* Notice that even thought SF_Distinct has been cleared from p->selFlags,
     ** the sDistinct.isTnct is still set.  Hence, isTnct represents the
     ** original setting of the SF_Distinct flag, not the current setting */
@@ -130684,7 +134070,7 @@ SQLITE_PRIVATE int sqlite3Select(
 #ifndef SQLITE_OMIT_WINDOWFUNC
     Window *pWin = p->pWin;      /* Master window object (or NULL) */
     if( pWin ){
-      sqlite3WindowCodeInit(pParse, pWin);
+      sqlite3WindowCodeInit(pParse, p);
     }
 #endif
     assert( WHERE_USE_LIMIT==SF_FixedLimit );
@@ -130779,23 +134165,35 @@ SQLITE_PRIVATE int sqlite3Select(
       }
       assert( 66==sqlite3LogEst(100) );
       if( p->nSelectRow>66 ) p->nSelectRow = 66;
+
+      /* If there is both a GROUP BY and an ORDER BY clause and they are
+      ** identical, then it may be possible to disable the ORDER BY clause 
+      ** on the grounds that the GROUP BY will cause elements to come out 
+      ** in the correct order. It also may not - the GROUP BY might use a
+      ** database index that causes rows to be grouped together as required
+      ** but not actually sorted. Either way, record the fact that the
+      ** ORDER BY and GROUP BY clauses are the same by setting the orderByGrp
+      ** variable.  */
+      if( sSort.pOrderBy && pGroupBy->nExpr==sSort.pOrderBy->nExpr ){
+        int ii;
+        /* The GROUP BY processing doesn't care whether rows are delivered in
+        ** ASC or DESC order - only that each group is returned contiguously.
+        ** So set the ASC/DESC flags in the GROUP BY to match those in the 
+        ** ORDER BY to maximize the chances of rows being delivered in an 
+        ** order that makes the ORDER BY redundant.  */
+        for(ii=0; ii<pGroupBy->nExpr; ii++){
+          u8 sortFlags = sSort.pOrderBy->a[ii].sortFlags & KEYINFO_ORDER_DESC;
+          pGroupBy->a[ii].sortFlags = sortFlags;
+        }
+        if( sqlite3ExprListCompare(pGroupBy, sSort.pOrderBy, -1)==0 ){
+          orderByGrp = 1;
+        }
+      }
     }else{
       assert( 0==sqlite3LogEst(1) );
       p->nSelectRow = 0;
     }
 
-    /* If there is both a GROUP BY and an ORDER BY clause and they are
-    ** identical, then it may be possible to disable the ORDER BY clause 
-    ** on the grounds that the GROUP BY will cause elements to come out 
-    ** in the correct order. It also may not - the GROUP BY might use a
-    ** database index that causes rows to be grouped together as required
-    ** but not actually sorted. Either way, record the fact that the
-    ** ORDER BY and GROUP BY clauses are the same by setting the orderByGrp
-    ** variable.  */
-    if( sqlite3ExprListCompare(pGroupBy, sSort.pOrderBy, -1)==0 ){
-      orderByGrp = 1;
-    }
- 
     /* Create a label to jump to when we want to abort the query */
     addrEnd = sqlite3VdbeMakeLabel(pParse);
 
@@ -130830,9 +134228,16 @@ SQLITE_PRIVATE int sqlite3Select(
       minMaxFlag = WHERE_ORDERBY_NORMAL;
     }
     for(i=0; i<sAggInfo.nFunc; i++){
-      assert( !ExprHasProperty(sAggInfo.aFunc[i].pExpr, EP_xIsSelect) );
+      Expr *pExpr = sAggInfo.aFunc[i].pExpr;
+      assert( !ExprHasProperty(pExpr, EP_xIsSelect) );
       sNC.ncFlags |= NC_InAggFunc;
-      sqlite3ExprAnalyzeAggList(&sNC, sAggInfo.aFunc[i].pExpr->x.pList);
+      sqlite3ExprAnalyzeAggList(&sNC, pExpr->x.pList);
+#ifndef SQLITE_OMIT_WINDOWFUNC
+      assert( !IsWindowFunc(pExpr) );
+      if( ExprHasProperty(pExpr, EP_WinFunc) ){
+        sqlite3ExprAnalyzeAggregates(&sNC, pExpr->y.pWin->pFilter);
+      }
+#endif
       sNC.ncFlags &= ~NC_InAggFunc;
     }
     sAggInfo.mxReg = pParse->nMem;
@@ -131144,13 +134549,18 @@ SQLITE_PRIVATE int sqlite3Select(
       {
         int regAcc = 0;           /* "populate accumulators" flag */
 
-        /* If there are accumulator registers but no min() or max() functions,
-        ** allocate register regAcc. Register regAcc will contain 0 the first
-        ** time the inner loop runs, and 1 thereafter. The code generated
-        ** by updateAccumulator() only updates the accumulator registers if
-        ** regAcc contains 0.  */
+        /* If there are accumulator registers but no min() or max() functions
+        ** without FILTER clauses, allocate register regAcc. Register regAcc
+        ** will contain 0 the first time the inner loop runs, and 1 thereafter.
+        ** The code generated by updateAccumulator() uses this to ensure
+        ** that the accumulator registers are (a) updated only once if
+        ** there are no min() or max functions or (b) always updated for the
+        ** first row visited by the aggregate, so that they are updated at
+        ** least once even if the FILTER clause means the min() or max() 
+        ** function visits zero rows.  */
         if( sAggInfo.nAccumulator ){
           for(i=0; i<sAggInfo.nFunc; i++){
+            if( ExprHasProperty(sAggInfo.aFunc[i].pExpr, EP_WinFunc) ) continue;
             if( sAggInfo.aFunc[i].pFunc->funcFlags&SQLITE_FUNC_NEEDCOLL ) break;
           }
           if( i==sAggInfo.nFunc ){
@@ -131621,7 +135031,11 @@ SQLITE_PRIVATE void sqlite3BeginTrigger(
   /* Check that the trigger name is not reserved and that no trigger of the
   ** specified name exists */
   zName = sqlite3NameFromToken(db, pName);
-  if( !zName || SQLITE_OK!=sqlite3CheckObjectName(pParse, zName) ){
+  if( zName==0 ){
+    assert( db->mallocFailed );
+    goto trigger_cleanup;
+  }
+  if( sqlite3CheckObjectName(pParse, zName, "trigger", pTab->zName) ){
     goto trigger_cleanup;
   }
   assert( sqlite3SchemaMutexHeld(db, iDb, 0) );
@@ -131784,6 +135198,7 @@ SQLITE_PRIVATE void sqlite3FinishTrigger(
     Trigger *pLink = pTrig;
     Hash *pHash = &db->aDb[iDb].pSchema->trigHash;
     assert( sqlite3SchemaMutexHeld(db, iDb, 0) );
+    assert( pLink!=0 );
     pTrig = sqlite3HashInsert(pHash, zName, pTrig);
     if( pTrig ){
       sqlite3OomFault(db);
@@ -131902,6 +135317,9 @@ SQLITE_PRIVATE TriggerStep *sqlite3TriggerInsertStep(
     pTriggerStep->pIdList = pColumn;
     pTriggerStep->pUpsert = pUpsert;
     pTriggerStep->orconf = orconf;
+    if( pUpsert ){
+      sqlite3HasExplicitNulls(pParse, pUpsert->pUpsertTarget);
+    }
   }else{
     testcase( pColumn );
     sqlite3IdListDelete(db, pColumn);
@@ -132057,10 +135475,9 @@ SQLITE_PRIVATE void sqlite3DropTriggerPtr(Parse *pParse, Trigger *pTrigger){
   iDb = sqlite3SchemaToIndex(pParse->db, pTrigger->pSchema);
   assert( iDb>=0 && iDb<db->nDb );
   pTable = tableOfTrigger(pTrigger);
-  assert( pTable );
-  assert( pTable->pSchema==pTrigger->pSchema || iDb==1 );
+  assert( (pTable && pTable->pSchema==pTrigger->pSchema) || iDb==1 );
 #ifndef SQLITE_OMIT_AUTHORIZATION
-  {
+  if( pTable ){
     int code = SQLITE_DROP_TRIGGER;
     const char *zDb = db->aDb[iDb].zDbSName;
     const char *zTab = SCHEMA_TABLE(iDb);
@@ -132074,7 +135491,6 @@ SQLITE_PRIVATE void sqlite3DropTriggerPtr(Parse *pParse, Trigger *pTrigger){
 
   /* Generate code to destroy the database record of the trigger.
   */
-  assert( pTable!=0 );
   if( (v = sqlite3GetVdbe(pParse))!=0 ){
     sqlite3NestedParse(pParse,
        "DELETE FROM %Q.%s WHERE name=%Q AND type='trigger'",
@@ -132098,9 +135514,15 @@ SQLITE_PRIVATE void sqlite3UnlinkAndDeleteTrigger(sqlite3 *db, int iDb, const ch
   if( ALWAYS(pTrigger) ){
     if( pTrigger->pSchema==pTrigger->pTabSchema ){
       Table *pTab = tableOfTrigger(pTrigger);
-      Trigger **pp;
-      for(pp=&pTab->pTrigger; *pp!=pTrigger; pp=&((*pp)->pNext));
-      *pp = (*pp)->pNext;
+      if( pTab ){
+        Trigger **pp;
+        for(pp=&pTab->pTrigger; *pp; pp=&((*pp)->pNext)){
+          if( *pp==pTrigger ){
+            *pp = (*pp)->pNext;
+            break;
+          }
+        }
+      }
     }
     sqlite3DeleteTrigger(db, pTrigger);
     db->mDbFlags |= DBFLAG_SchemaChange;
@@ -132120,7 +135542,7 @@ static int checkColumnOverlap(IdList *pIdList, ExprList *pEList){
   int e;
   if( pIdList==0 || NEVER(pEList==0) ) return 1;
   for(e=0; e<pEList->nExpr; e++){
-    if( sqlite3IdListIndex(pIdList, pEList->a[e].zName)>=0 ) return 1;
+    if( sqlite3IdListIndex(pIdList, pEList->a[e].zEName)>=0 ) return 1;
   }
   return 0; 
 }
@@ -132780,7 +136202,7 @@ SQLITE_PRIVATE void sqlite3Update(
   Expr *pLimit,          /* LIMIT clause. May be null */
   Upsert *pUpsert        /* ON CONFLICT clause, or null */
 ){
-  int i, j;              /* Loop counters */
+  int i, j, k;           /* Loop counters */
   Table *pTab;           /* The table to be updated */
   int addrTop = 0;       /* VDBE instruction address of the start of the loop */
   WhereInfo *pWInfo;     /* Information about the WHERE clause */
@@ -132788,11 +136210,12 @@ SQLITE_PRIVATE void sqlite3Update(
   Index *pIdx;           /* For looping over indices */
   Index *pPk;            /* The PRIMARY KEY index for WITHOUT ROWID tables */
   int nIdx;              /* Number of indices that need updating */
+  int nAllIdx;           /* Total number of indexes */
   int iBaseCur;          /* Base cursor number */
   int iDataCur;          /* Cursor for the canonical data btree */
   int iIdxCur;           /* Cursor for the first index */
   sqlite3 *db;           /* The database structure */
-  int *aRegIdx = 0;      /* First register in array assigned to each index */
+  int *aRegIdx = 0;      /* Registers for to each index and the main table */
   int *aXRef = 0;        /* aXRef[i] is the index in pChanges->a[] of the
                          ** an expression for the i-th column of the table.
                          ** aXRef[i]==-1 if the i-th column is not changed. */
@@ -132823,6 +136246,7 @@ SQLITE_PRIVATE void sqlite3Update(
   int iPk = 0;           /* First of nPk cells holding PRIMARY KEY value */
   i16 nPk = 0;           /* Number of components of the PRIMARY KEY */
   int bReplace = 0;      /* True if REPLACE conflict resolution might happen */
+  int bFinishSeek = 1;   /* The OP_FinishSeek opcode is needed */
 
   /* Register Allocations */
   int regRowCount = 0;   /* A count of rows changed */
@@ -132906,10 +136330,10 @@ SQLITE_PRIVATE void sqlite3Update(
   /* Allocate space for aXRef[], aRegIdx[], and aToOpen[].  
   ** Initialize aXRef[] and aToOpen[] to their default values.
   */
-  aXRef = sqlite3DbMallocRawNN(db, sizeof(int) * (pTab->nCol+nIdx) + nIdx+2 );
+  aXRef = sqlite3DbMallocRawNN(db, sizeof(int) * (pTab->nCol+nIdx+1) + nIdx+2 );
   if( aXRef==0 ) goto update_cleanup;
   aRegIdx = aXRef+pTab->nCol;
-  aToOpen = (u8*)(aRegIdx+nIdx);
+  aToOpen = (u8*)(aRegIdx+nIdx+1);
   memset(aToOpen, 1, nIdx+1);
   aToOpen[nIdx+1] = 0;
   for(i=0; i<pTab->nCol; i++) aXRef[i] = -1;
@@ -132921,6 +136345,10 @@ SQLITE_PRIVATE void sqlite3Update(
   sNC.uNC.pUpsert = pUpsert;
   sNC.ncFlags = NC_UUpsert;
 
+  /* Begin generating code. */
+  v = sqlite3GetVdbe(pParse);
+  if( v==0 ) goto update_cleanup;
+
   /* Resolve the column names in all the expressions of the
   ** of the UPDATE statement.  Also find the column index
   ** for each column to be updated in the pChanges array.  For each
@@ -132933,24 +136361,34 @@ SQLITE_PRIVATE void sqlite3Update(
       goto update_cleanup;
     }
     for(j=0; j<pTab->nCol; j++){
-      if( sqlite3StrICmp(pTab->aCol[j].zName, pChanges->a[i].zName)==0 ){
+      if( sqlite3StrICmp(pTab->aCol[j].zName, pChanges->a[i].zEName)==0 ){
         if( j==pTab->iPKey ){
           chngRowid = 1;
           pRowidExpr = pChanges->a[i].pExpr;
         }else if( pPk && (pTab->aCol[j].colFlags & COLFLAG_PRIMKEY)!=0 ){
           chngPk = 1;
         }
+#ifndef SQLITE_OMIT_GENERATED_COLUMNS
+        else if( pTab->aCol[j].colFlags & COLFLAG_GENERATED ){
+          testcase( pTab->aCol[j].colFlags & COLFLAG_VIRTUAL );
+          testcase( pTab->aCol[j].colFlags & COLFLAG_STORED );
+          sqlite3ErrorMsg(pParse, 
+             "cannot UPDATE generated column \"%s\"",
+             pTab->aCol[j].zName);
+          goto update_cleanup;
+        }
+#endif
         aXRef[j] = i;
         break;
       }
     }
     if( j>=pTab->nCol ){
-      if( pPk==0 && sqlite3IsRowid(pChanges->a[i].zName) ){
+      if( pPk==0 && sqlite3IsRowid(pChanges->a[i].zEName) ){
         j = -1;
         chngRowid = 1;
         pRowidExpr = pChanges->a[i].pExpr;
       }else{
-        sqlite3ErrorMsg(pParse, "no such column: %s", pChanges->a[i].zName);
+        sqlite3ErrorMsg(pParse, "no such column: %s", pChanges->a[i].zEName);
         pParse->checkSchema = 1;
         goto update_cleanup;
       }
@@ -132974,6 +136412,33 @@ SQLITE_PRIVATE void sqlite3Update(
   assert( chngPk==0 || chngPk==1 );
   chngKey = chngRowid + chngPk;
 
+#ifndef SQLITE_OMIT_GENERATED_COLUMNS
+  /* Mark generated columns as changing if their generator expressions
+  ** reference any changing column.  The actual aXRef[] value for 
+  ** generated expressions is not used, other than to check to see that it
+  ** is non-negative, so the value of aXRef[] for generated columns can be
+  ** set to any non-negative number.  We use 99999 so that the value is
+  ** obvious when looking at aXRef[] in a symbolic debugger. 
+  */
+  if( pTab->tabFlags & TF_HasGenerated ){
+    int bProgress;
+    testcase( pTab->tabFlags & TF_HasVirtual );
+    testcase( pTab->tabFlags & TF_HasStored );
+    do{
+      bProgress = 0;
+      for(i=0; i<pTab->nCol; i++){
+        if( aXRef[i]>=0 ) continue;
+        if( (pTab->aCol[i].colFlags & COLFLAG_GENERATED)==0 ) continue;
+        if( sqlite3ExprReferencesUpdatedColumn(pTab->aCol[i].pDflt,
+                                               aXRef, chngRowid) ){
+          aXRef[i] = 99999;
+          bProgress = 1;
+        }
+      }
+    }while( bProgress );
+  }
+#endif
+
   /* The SET expressions are not actually used inside the WHERE loop.  
   ** So reset the colUsed mask. Unless this is a virtual table. In that
   ** case, set all bits of the colUsed mask (to ensure that the virtual
@@ -132988,7 +136453,7 @@ SQLITE_PRIVATE void sqlite3Update(
   ** the key for accessing each index.
   */
   if( onError==OE_Replace ) bReplace = 1;
-  for(j=0, pIdx=pTab->pIndex; pIdx; pIdx=pIdx->pNext, j++){
+  for(nAllIdx=0, pIdx=pTab->pIndex; pIdx; pIdx=pIdx->pNext, nAllIdx++){
     int reg;
     if( chngKey || hasFK>1 || pIdx==pPk
      || indexWhereClauseMightChange(pIdx,aXRef,chngRowid)
@@ -133008,24 +136473,28 @@ SQLITE_PRIVATE void sqlite3Update(
         }
       }
     }
-    if( reg==0 ) aToOpen[j+1] = 0;
-    aRegIdx[j] = reg;
+    if( reg==0 ) aToOpen[nAllIdx+1] = 0;
+    aRegIdx[nAllIdx] = reg;
   }
+  aRegIdx[nAllIdx] = ++pParse->nMem;  /* Register storing the table record */
   if( bReplace ){
     /* If REPLACE conflict resolution might be invoked, open cursors on all 
     ** indexes in case they are needed to delete records.  */
     memset(aToOpen, 1, nIdx+1);
   }
 
-  /* Begin generating code. */
-  v = sqlite3GetVdbe(pParse);
-  if( v==0 ) goto update_cleanup;
   if( pParse->nested==0 ) sqlite3VdbeCountChanges(v);
   sqlite3BeginWriteOperation(pParse, pTrigger || hasFK, iDb);
 
   /* Allocate required registers. */
   if( !IsVirtual(pTab) ){
-    regRowSet = ++pParse->nMem;
+    /* For now, regRowSet and aRegIdx[nAllIdx] share the same register.
+    ** If regRowSet turns out to be needed, then aRegIdx[nAllIdx] will be
+    ** reallocated.  aRegIdx[nAllIdx] is the register in which the main
+    ** table record is written.  regRowSet holds the RowSet for the
+    ** two-pass update algorithm. */
+    assert( aRegIdx[nAllIdx]==pParse->nMem );
+    regRowSet = aRegIdx[nAllIdx];
     regOldRowid = regNewRowid = ++pParse->nMem;
     if( chngPk || pTrigger || hasFK ){
       regOld = pParse->nMem + 1;
@@ -133111,6 +136580,7 @@ SQLITE_PRIVATE void sqlite3Update(
     pWInfo = 0;
     eOnePass = ONEPASS_SINGLE;
     sqlite3ExprIfFalse(pParse, pWhere, labelBreak, SQLITE_JUMPIFNULL);
+    bFinishSeek = 0;
   }else{
     /* Begin the database scan. 
     **
@@ -133137,6 +136607,7 @@ SQLITE_PRIVATE void sqlite3Update(
     ** strategy that uses an index for which one or more columns are being
     ** updated.  */
     eOnePass = sqlite3WhereOkOnePass(pWInfo, aiCurOnePass);
+    bFinishSeek = sqlite3WhereUsesDeferredSeek(pWInfo);
     if( eOnePass!=ONEPASS_SINGLE ){
       sqlite3MultiWrite(pParse);
       if( eOnePass==ONEPASS_MULTI ){
@@ -133155,6 +136626,8 @@ SQLITE_PRIVATE void sqlite3Update(
     ** leave it in register regOldRowid.  */
     sqlite3VdbeAddOp2(v, OP_Rowid, iDataCur, regOldRowid);
     if( eOnePass==ONEPASS_OFF ){
+      /* We need to use regRowSet, so reallocate aRegIdx[nAllIdx] */
+      aRegIdx[nAllIdx] = ++pParse->nMem;
       sqlite3VdbeAddOp2(v, OP_RowSetAdd, regRowSet, regOldRowid);
     }
   }else{
@@ -133165,7 +136638,8 @@ SQLITE_PRIVATE void sqlite3Update(
     ** is not required) and leave the PK fields in the array of registers.  */
     for(i=0; i<nPk; i++){
       assert( pPk->aiColumn[i]>=0 );
-      sqlite3ExprCodeGetColumnOfTable(v, pTab, iDataCur,pPk->aiColumn[i],iPk+i);
+      sqlite3ExprCodeGetColumnOfTable(v, pTab, iDataCur,
+                                      pPk->aiColumn[i], iPk+i);
     }
     if( eOnePass ){
       if( addrOpen ) sqlite3VdbeChangeToNoop(v, addrOpen);
@@ -133246,14 +136720,16 @@ SQLITE_PRIVATE void sqlite3Update(
         pTrigger, pChanges, 0, TRIGGER_BEFORE|TRIGGER_AFTER, pTab, onError
     );
     for(i=0; i<pTab->nCol; i++){
+      u32 colFlags = pTab->aCol[i].colFlags;
+      k = sqlite3TableColumnToStorage(pTab, i) + regOld;
       if( oldmask==0xffffffff
        || (i<32 && (oldmask & MASKBIT32(i))!=0)
-       || (pTab->aCol[i].colFlags & COLFLAG_PRIMKEY)!=0
+       || (colFlags & COLFLAG_PRIMKEY)!=0
       ){
         testcase(  oldmask!=0xffffffff && i==31 );
-        sqlite3ExprCodeGetColumnOfTable(v, pTab, iDataCur, i, regOld+i);
+        sqlite3ExprCodeGetColumnOfTable(v, pTab, iDataCur, i, k);
       }else{
-        sqlite3VdbeAddOp2(v, OP_Null, 0, regOld+i);
+        sqlite3VdbeAddOp2(v, OP_Null, 0, k);
       }
     }
     if( chngRowid==0 && pPk==0 ){
@@ -133277,13 +136753,15 @@ SQLITE_PRIVATE void sqlite3Update(
   newmask = sqlite3TriggerColmask(
       pParse, pTrigger, pChanges, 1, TRIGGER_BEFORE, pTab, onError
   );
-  for(i=0; i<pTab->nCol; i++){
+  for(i=0, k=regNew; i<pTab->nCol; i++, k++){
     if( i==pTab->iPKey ){
-      sqlite3VdbeAddOp2(v, OP_Null, 0, regNew+i);
+      sqlite3VdbeAddOp2(v, OP_Null, 0, k);
+    }else if( (pTab->aCol[i].colFlags & COLFLAG_GENERATED)!=0 ){
+      if( pTab->aCol[i].colFlags & COLFLAG_VIRTUAL ) k--;
     }else{
       j = aXRef[i];
       if( j>=0 ){
-        sqlite3ExprCode(pParse, pChanges->a[j].pExpr, regNew+i);
+        sqlite3ExprCode(pParse, pChanges->a[j].pExpr, k);
       }else if( 0==(tmask&TRIGGER_BEFORE) || i>31 || (newmask & MASKBIT32(i)) ){
         /* This branch loads the value of a column that will not be changed 
         ** into a register. This is done if there are no BEFORE triggers, or
@@ -133292,12 +136770,20 @@ SQLITE_PRIVATE void sqlite3Update(
         */
         testcase( i==31 );
         testcase( i==32 );
-        sqlite3ExprCodeGetColumnOfTable(v, pTab, iDataCur, i, regNew+i);
+        sqlite3ExprCodeGetColumnOfTable(v, pTab, iDataCur, i, k);
+        bFinishSeek = 0;
       }else{
-        sqlite3VdbeAddOp2(v, OP_Null, 0, regNew+i);
+        sqlite3VdbeAddOp2(v, OP_Null, 0, k);
       }
     }
   }
+#ifndef SQLITE_OMIT_GENERATED_COLUMNS
+  if( pTab->tabFlags & TF_HasGenerated ){
+    testcase( pTab->tabFlags & TF_HasVirtual );
+    testcase( pTab->tabFlags & TF_HasStored );
+    sqlite3ComputeGeneratedColumns(pParse, regNew, pTab);
+  }
+#endif
 
   /* Fire any BEFORE UPDATE triggers. This happens before constraints are
   ** verified. One could argue that this is wrong.
@@ -133330,38 +136816,58 @@ SQLITE_PRIVATE void sqlite3Update(
     ** BEFORE trigger runs.  See test case trigger1-18.0 (added 2018-04-26)
     ** for an example.
     */
-    for(i=0; i<pTab->nCol; i++){
-      if( aXRef[i]<0 && i!=pTab->iPKey ){
-        sqlite3ExprCodeGetColumnOfTable(v, pTab, iDataCur, i, regNew+i);
+    for(i=0, k=regNew; i<pTab->nCol; i++, k++){
+      if( pTab->aCol[i].colFlags & COLFLAG_GENERATED ){
+        if( pTab->aCol[i].colFlags & COLFLAG_VIRTUAL ) k--;
+      }else if( aXRef[i]<0 && i!=pTab->iPKey ){
+        sqlite3ExprCodeGetColumnOfTable(v, pTab, iDataCur, i, k);
       }
     }
+#ifndef SQLITE_OMIT_GENERATED_COLUMNS
+    if( pTab->tabFlags & TF_HasGenerated ){
+      testcase( pTab->tabFlags & TF_HasVirtual );
+      testcase( pTab->tabFlags & TF_HasStored );
+      sqlite3ComputeGeneratedColumns(pParse, regNew, pTab);
+    }
+#endif 
   }
 
   if( !isView ){
-    int addr1 = 0;        /* Address of jump instruction */
-
     /* Do constraint checks. */
     assert( regOldRowid>0 );
     sqlite3GenerateConstraintChecks(pParse, pTab, aRegIdx, iDataCur, iIdxCur,
         regNewRowid, regOldRowid, chngKey, onError, labelContinue, &bReplace,
         aXRef, 0);
 
-    /* Do FK constraint checks. */
-    if( hasFK ){
-      sqlite3FkCheck(pParse, pTab, regOldRowid, 0, aXRef, chngKey);
-    }
-
-    /* Delete the index entries associated with the current record.  */
+    /* If REPLACE conflict handling may have been used, or if the PK of the
+    ** row is changing, then the GenerateConstraintChecks() above may have
+    ** moved cursor iDataCur. Reseek it. */
     if( bReplace || chngKey ){
       if( pPk ){
-        addr1 = sqlite3VdbeAddOp4Int(v, OP_NotFound, iDataCur, 0, regKey, nKey);
+        sqlite3VdbeAddOp4Int(v, OP_NotFound,iDataCur,labelContinue,regKey,nKey);
       }else{
-        addr1 = sqlite3VdbeAddOp3(v, OP_NotExists, iDataCur, 0, regOldRowid);
+        sqlite3VdbeAddOp3(v, OP_NotExists, iDataCur, labelContinue,regOldRowid);
       }
       VdbeCoverageNeverTaken(v);
     }
+
+    /* Do FK constraint checks. */
+    if( hasFK ){
+      sqlite3FkCheck(pParse, pTab, regOldRowid, 0, aXRef, chngKey);
+    }
+
+    /* Delete the index entries associated with the current record.  */
     sqlite3GenerateRowIndexDelete(pParse, pTab, iDataCur, iIdxCur, aRegIdx, -1);
 
+    /* We must run the OP_FinishSeek opcode to resolve a prior
+    ** OP_DeferredSeek if there is any possibility that there have been
+    ** no OP_Column opcodes since the OP_DeferredSeek was issued.  But
+    ** we want to avoid the OP_FinishSeek if possible, as running it
+    ** costs CPU cycles. */
+    if( bFinishSeek ){
+      sqlite3VdbeAddOp1(v, OP_FinishSeek, iDataCur);
+    }
+
     /* If changing the rowid value, or if there are foreign key constraints
     ** to process, delete the old record. Otherwise, add a noop OP_Delete
     ** to invoke the pre-update hook.
@@ -133389,9 +136895,6 @@ SQLITE_PRIVATE void sqlite3Update(
       sqlite3VdbeAddOp2(v, OP_Delete, iDataCur, 0);
     }
 #endif
-    if( bReplace || chngKey ){
-      sqlite3VdbeJumpHere(v, addr1);
-    }
 
     if( hasFK ){
       sqlite3FkCheck(pParse, pTab, 0, regNewRowid, aXRef, chngKey);
@@ -133541,6 +137044,7 @@ static void updateVirtualTable(
 
   /* Populate the argument registers. */
   for(i=0; i<pTab->nCol; i++){
+    assert( (pTab->aCol[i].colFlags & COLFLAG_GENERATED)==0 );
     if( aXRef[i]>=0 ){
       sqlite3ExprCode(pParse, pChanges->a[aXRef[i]].pExpr, regArg+2+i);
     }else{
@@ -133830,6 +137334,7 @@ SQLITE_PRIVATE void sqlite3UpsertDoUpdate(
   sqlite3 *db = pParse->db;
   SrcList *pSrc;            /* FROM clause for the UPDATE */
   int iDataCur;
+  int i;
 
   assert( v!=0 );
   assert( pUpsert!=0 );
@@ -133846,12 +137351,11 @@ SQLITE_PRIVATE void sqlite3UpsertDoUpdate(
       Index *pPk = sqlite3PrimaryKeyIndex(pTab);
       int nPk = pPk->nKeyCol;
       int iPk = pParse->nMem+1;
-      int i;
       pParse->nMem += nPk;
       for(i=0; i<nPk; i++){
         int k;
         assert( pPk->aiColumn[i]>=0 );
-        k = sqlite3ColumnOfIndex(pIdx, pPk->aiColumn[i]);
+        k = sqlite3TableColumnToIndex(pIdx, pPk->aiColumn[i]);
         sqlite3VdbeAddOp3(v, OP_Column, iCur, k, iPk+i);
         VdbeComment((v, "%s.%s", pIdx->zName,
                     pTab->aCol[pPk->aiColumn[i]].zName));
@@ -133861,12 +137365,19 @@ SQLITE_PRIVATE void sqlite3UpsertDoUpdate(
       VdbeCoverage(v);
       sqlite3VdbeAddOp4(v, OP_Halt, SQLITE_CORRUPT, OE_Abort, 0, 
             "corrupt database", P4_STATIC);
+      sqlite3MayAbort(pParse);
       sqlite3VdbeJumpHere(v, i);
     }
   }
   /* pUpsert does not own pUpsertSrc - the outer INSERT statement does.  So
   ** we have to make a copy before passing it down into sqlite3Update() */
   pSrc = sqlite3SrcListDup(db, pUpsert->pUpsertSrc, 0);
+  /* excluded.* columns of type REAL need to be converted to a hard real */
+  for(i=0; i<pTab->nCol; i++){
+    if( pTab->aCol[i].affinity==SQLITE_AFF_REAL ){
+      sqlite3VdbeAddOp1(v, OP_RealAffinity, pUpsert->regData+i);
+    }
+  }
   sqlite3Update(pParse, pSrc, pUpsert->pUpsertSet,
       pUpsert->pUpsertWhere, OE_Abort, 0, 0, pUpsert);
   pUpsert->pUpsertSet = 0;    /* Will have been deleted by sqlite3Update() */
@@ -133986,6 +137497,7 @@ SQLITE_PRIVATE void sqlite3Vacuum(Parse *pParse, Token *pNm, Expr *pInto){
   Vdbe *v = sqlite3GetVdbe(pParse);
   int iDb = 0;
   if( v==0 ) goto build_vacuum_end;
+  if( pParse->nErr ) goto build_vacuum_end;
   if( pNm ){
 #ifndef SQLITE_BUG_COMPATIBLE_20160819
     /* Default behavior:  Report an error if the argument to VACUUM is
@@ -134329,6 +137841,9 @@ struct VtabCtx {
 ** Construct and install a Module object for a virtual table.  When this
 ** routine is called, it is guaranteed that all appropriate locks are held
 ** and the module is not already part of the connection.
+**
+** If there already exists a module with zName, replace it with the new one.
+** If pModule==0, then delete the module zName if it exists.
 */
 SQLITE_PRIVATE Module *sqlite3VtabCreateModule(
   sqlite3 *db,                    /* Database in which module is registered */
@@ -134338,25 +137853,36 @@ SQLITE_PRIVATE Module *sqlite3VtabCreateModule(
   void (*xDestroy)(void *)        /* Module destructor function */
 ){
   Module *pMod;
-  int nName = sqlite3Strlen30(zName);
-  pMod = (Module *)sqlite3Malloc(sizeof(Module) + nName + 1);
-  if( pMod==0 ){
-    sqlite3OomFault(db);
+  Module *pDel;
+  char *zCopy;
+  if( pModule==0 ){
+    zCopy = (char*)zName;
+    pMod = 0;
   }else{
-    Module *pDel;
-    char *zCopy = (char *)(&pMod[1]);
+    int nName = sqlite3Strlen30(zName);
+    pMod = (Module *)sqlite3Malloc(sizeof(Module) + nName + 1);
+    if( pMod==0 ){
+      sqlite3OomFault(db);
+      return 0;
+    }
+    zCopy = (char *)(&pMod[1]);
     memcpy(zCopy, zName, nName+1);
     pMod->zName = zCopy;
     pMod->pModule = pModule;
     pMod->pAux = pAux;
     pMod->xDestroy = xDestroy;
     pMod->pEpoTab = 0;
-    pDel = (Module *)sqlite3HashInsert(&db->aModule,zCopy,(void*)pMod);
-    assert( pDel==0 || pDel==pMod );
-    if( pDel ){
+    pMod->nRefModule = 1;
+  }
+  pDel = (Module *)sqlite3HashInsert(&db->aModule,zCopy,(void*)pMod);
+  if( pDel ){
+    if( pDel==pMod ){
       sqlite3OomFault(db);
       sqlite3DbFree(db, pDel);
       pMod = 0;
+    }else{
+      sqlite3VtabEponymousTableClear(db, pDel);
+      sqlite3VtabModuleUnref(db, pDel);
     }
   }
   return pMod;
@@ -134377,11 +137903,7 @@ static int createModule(
   int rc = SQLITE_OK;
 
   sqlite3_mutex_enter(db->mutex);
-  if( sqlite3HashFind(&db->aModule, zName) ){
-    rc = SQLITE_MISUSE_BKPT;
-  }else{
-    (void)sqlite3VtabCreateModule(db, zName, pModule, pAux, xDestroy);
-  }
+  (void)sqlite3VtabCreateModule(db, zName, pModule, pAux, xDestroy);
   rc = sqlite3ApiExit(db, rc);
   if( rc!=SQLITE_OK && xDestroy ) xDestroy(pAux);
   sqlite3_mutex_leave(db->mutex);
@@ -134420,6 +137942,44 @@ SQLITE_API int sqlite3_create_module_v2(
   return createModule(db, zName, pModule, pAux, xDestroy);
 }
 
+/*
+** External API to drop all virtual-table modules, except those named
+** on the azNames list.
+*/
+SQLITE_API int sqlite3_drop_modules(sqlite3 *db, const char** azNames){
+  HashElem *pThis, *pNext;
+#ifdef SQLITE_ENABLE_API_ARMOR
+  if( !sqlite3SafetyCheckOk(db) ) return SQLITE_MISUSE_BKPT;
+#endif
+  for(pThis=sqliteHashFirst(&db->aModule); pThis; pThis=pNext){
+    Module *pMod = (Module*)sqliteHashData(pThis);
+    pNext = sqliteHashNext(pThis);
+    if( azNames ){
+      int ii;
+      for(ii=0; azNames[ii]!=0 && strcmp(azNames[ii],pMod->zName)!=0; ii++){}
+      if( azNames[ii]!=0 ) continue;
+    }
+    createModule(db, pMod->zName, 0, 0, 0);
+  }
+  return SQLITE_OK;
+}
+
+/*
+** Decrement the reference count on a Module object.  Destroy the
+** module when the reference count reaches zero.
+*/
+SQLITE_PRIVATE void sqlite3VtabModuleUnref(sqlite3 *db, Module *pMod){
+  assert( pMod->nRefModule>0 );
+  pMod->nRefModule--;
+  if( pMod->nRefModule==0 ){
+    if( pMod->xDestroy ){
+      pMod->xDestroy(pMod->pAux);
+    }
+    assert( pMod->pEpoTab==0 );
+    sqlite3DbFree(db, pMod);
+  }
+}
+
 /*
 ** Lock the virtual table so that it cannot be disconnected.
 ** Locks nest.  Every lock should have a corresponding unlock.
@@ -134459,6 +138019,7 @@ SQLITE_PRIVATE void sqlite3VtabUnlock(VTable *pVTab){
   pVTab->nRef--;
   if( pVTab->nRef==0 ){
     sqlite3_vtab *p = pVTab->pVtab;
+    sqlite3VtabModuleUnref(pVTab->db, pVTab->pMod);
     if( p ){
       p->pModule->xDisconnect(p);
     }
@@ -134553,12 +138114,12 @@ SQLITE_PRIVATE void sqlite3VtabDisconnect(sqlite3 *db, Table *p){
 */
 SQLITE_PRIVATE void sqlite3VtabUnlockList(sqlite3 *db){
   VTable *p = db->pDisconnect;
-  db->pDisconnect = 0;
 
   assert( sqlite3BtreeHoldsAllMutexes(db) );
   assert( sqlite3_mutex_held(db->mutex) );
 
   if( p ){
+    db->pDisconnect = 0;
     sqlite3ExpirePreparedStatements(db, 0);
     do {
       VTable *pNext = p->pNext;
@@ -134705,6 +138266,8 @@ SQLITE_PRIVATE void sqlite3VtabFinishParse(Parse *pParse, Token *pEnd){
     int iReg;
     Vdbe *v;
 
+    sqlite3MayAbort(pParse);
+
     /* Compute the complete text of the CREATE VIRTUAL TABLE statement */
     if( pEnd ){
       pParse->sNameToken.n = (int)(pEnd->z - pParse->sNameToken.z) + pEnd->n;
@@ -134730,13 +138293,13 @@ SQLITE_PRIVATE void sqlite3VtabFinishParse(Parse *pParse, Token *pEnd){
       zStmt,
       pParse->regRowid
     );
-    sqlite3DbFree(db, zStmt);
     v = sqlite3GetVdbe(pParse);
     sqlite3ChangeCookie(pParse, iDb);
 
     sqlite3VdbeAddOp0(v, OP_Expire);
-    zWhere = sqlite3MPrintf(db, "name='%q' AND type='table'", pTab->zName);
+    zWhere = sqlite3MPrintf(db, "name=%Q AND sql=%Q", pTab->zName, zStmt);
     sqlite3VdbeAddParseSchemaOp(v, iDb, zWhere);
+    sqlite3DbFree(db, zStmt);
 
     iReg = ++pParse->nMem;
     sqlite3VdbeLoadString(v, iReg, pTab->zName);
@@ -134833,6 +138396,7 @@ static int vtabCallConstructor(
   }
   pVTable->db = db;
   pVTable->pMod = pMod;
+  pVTable->eVtabRisk = SQLITE_VTABRISK_Normal;
 
   iDb = sqlite3SchemaToIndex(db, pTab->pSchema);
   pTab->azModuleArg[1] = db->aDb[iDb].zDbSName;
@@ -134863,6 +138427,7 @@ static int vtabCallConstructor(
     ** the sqlite3_vtab object if successful.  */
     memset(pVTable->pVtab, 0, sizeof(pVTable->pVtab[0]));
     pVTable->pVtab->pModule = pMod->pModule;
+    pMod->nRefModule++;
     pVTable->nRef = 1;
     if( sCtx.bDeclared==0 ){
       const char *zFormat = "vtable constructor did not declare schema: %s";
@@ -134871,7 +138436,7 @@ static int vtabCallConstructor(
       rc = SQLITE_ERROR;
     }else{
       int iCol;
-      u8 oooHidden = 0;
+      u16 oooHidden = 0;
       /* If everything went according to plan, link the new VTable structure
       ** into the linked list headed by pTab->pVTable. Then loop through the 
       ** columns of the table to see if any of them contain the token "hidden".
@@ -135137,7 +138702,9 @@ SQLITE_PRIVATE int sqlite3VtabCallDestroy(sqlite3 *db, int iDb, const char *zTab
     }
     p = vtabDisconnectAll(db, pTab);
     xDestroy = p->pMod->pModule->xDestroy;
-    assert( xDestroy!=0 );  /* Checked before the virtual table is created */
+    if( xDestroy==0 ) xDestroy = p->pMod->pModule->xDisconnect;
+    assert( xDestroy!=0 );
+    pTab->nTabRef++;
     rc = xDestroy(p->pVtab);
     /* Remove the sqlite3_vtab* from the aVTrans[] array, if applicable */
     if( rc==SQLITE_OK ){
@@ -135146,6 +138713,7 @@ SQLITE_PRIVATE int sqlite3VtabCallDestroy(sqlite3 *db, int iDb, const char *zTab
       pTab->pVTable = 0;
       sqlite3VtabUnlock(p);
     }
+    sqlite3DeleteTable(db, pTab);
   }
 
   return rc;
@@ -135518,28 +139086,38 @@ SQLITE_API int sqlite3_vtab_on_conflict(sqlite3 *db){
 SQLITE_API int sqlite3_vtab_config(sqlite3 *db, int op, ...){
   va_list ap;
   int rc = SQLITE_OK;
+  VtabCtx *p;
 
 #ifdef SQLITE_ENABLE_API_ARMOR
   if( !sqlite3SafetyCheckOk(db) ) return SQLITE_MISUSE_BKPT;
 #endif
   sqlite3_mutex_enter(db->mutex);
-  va_start(ap, op);
-  switch( op ){
-    case SQLITE_VTAB_CONSTRAINT_SUPPORT: {
-      VtabCtx *p = db->pVtabCtx;
-      if( !p ){
-        rc = SQLITE_MISUSE_BKPT;
-      }else{
-        assert( p->pTab==0 || IsVirtual(p->pTab) );
+  p = db->pVtabCtx;
+  if( !p ){
+    rc = SQLITE_MISUSE_BKPT;
+  }else{
+    assert( p->pTab==0 || IsVirtual(p->pTab) );
+    va_start(ap, op);
+    switch( op ){
+      case SQLITE_VTAB_CONSTRAINT_SUPPORT: {
         p->pVTable->bConstraint = (u8)va_arg(ap, int);
+        break;
+      }
+      case SQLITE_VTAB_INNOCUOUS: {
+        p->pVTable->eVtabRisk = SQLITE_VTABRISK_Low;
+        break;
+      }
+      case SQLITE_VTAB_DIRECTONLY: {
+        p->pVTable->eVtabRisk = SQLITE_VTABRISK_High;
+        break;
+      }
+      default: {
+        rc = SQLITE_MISUSE_BKPT;
+        break;
       }
-      break;
     }
-    default:
-      rc = SQLITE_MISUSE_BKPT;
-      break;
+    va_end(ap);
   }
-  va_end(ap);
 
   if( rc!=SQLITE_OK ) sqlite3Error(db, rc);
   sqlite3_mutex_leave(db->mutex);
@@ -135588,6 +139166,8 @@ SQLITE_API int sqlite3_vtab_config(sqlite3 *db, int op, ...){
 ** planner logic in "where.c".  These definitions are broken out into
 ** a separate source file for easier editing.
 */
+#ifndef SQLITE_WHEREINT_H
+#define SQLITE_WHEREINT_H
 
 /*
 ** Trace output macros
@@ -135643,13 +139223,15 @@ struct WhereLevel {
   int addrCont;         /* Jump here to continue with the next loop cycle */
   int addrFirst;        /* First instruction of interior of the loop */
   int addrBody;         /* Beginning of the body of this loop */
+  int regBignull;       /* big-null flag reg. True if a NULL-scan is needed */
+  int addrBignull;      /* Jump here for next part of big-null scan */
 #ifndef SQLITE_LIKE_DOESNT_MATCH_BLOBS
   u32 iLikeRepCntr;     /* LIKE range processing counter register (times 2) */
   int addrLikeRep;      /* LIKE range processing address */
 #endif
   u8 iFrom;             /* Which entry in the FROM clause */
   u8 op, p3, p5;        /* Opcode, P3 & P5 of the opcode that ends the loop */
-  int p1, p2;           /* Operands of the opcode used to ends the loop */
+  int p1, p2;           /* Operands of the opcode used to end the loop */
   union {               /* Information that depends on pWLoop->wsFlags */
     struct {
       int nIn;              /* Number of entries in aInLoop[] */
@@ -135700,7 +139282,7 @@ struct WhereLoop {
       u16 nEq;               /* Number of equality constraints */
       u16 nBtm;              /* Size of BTM vector */
       u16 nTop;              /* Size of TOP vector */
-      u16 nIdxCol;           /* Index column used for ORDER BY */
+      u16 nDistinctCol;      /* Index columns used to sort for DISTINCT */
       Index *pIndex;         /* Index used, or NULL */
     } btree;
     struct {               /* Information for virtual tables */
@@ -135844,22 +139426,22 @@ struct WhereTerm {
 /*
 ** Allowed values of WhereTerm.wtFlags
 */
-#define TERM_DYNAMIC    0x01   /* Need to call sqlite3ExprDelete(db, pExpr) */
-#define TERM_VIRTUAL    0x02   /* Added by the optimizer.  Do not code */
-#define TERM_CODED      0x04   /* This term is already coded */
-#define TERM_COPIED     0x08   /* Has a child */
-#define TERM_ORINFO     0x10   /* Need to free the WhereTerm.u.pOrInfo object */
-#define TERM_ANDINFO    0x20   /* Need to free the WhereTerm.u.pAndInfo obj */
-#define TERM_OR_OK      0x40   /* Used during OR-clause processing */
-#ifdef SQLITE_ENABLE_STAT3_OR_STAT4
-#  define TERM_VNULL    0x80   /* Manufactured x>NULL or x<=NULL term */
-#else
-#  define TERM_VNULL    0x00   /* Disabled if not using stat3 */
-#endif
-#define TERM_LIKEOPT    0x100  /* Virtual terms from the LIKE optimization */
-#define TERM_LIKECOND   0x200  /* Conditionally this LIKE operator term */
-#define TERM_LIKE       0x400  /* The original LIKE operator */
-#define TERM_IS         0x800  /* Term.pExpr is an IS operator */
+#define TERM_DYNAMIC    0x0001 /* Need to call sqlite3ExprDelete(db, pExpr) */
+#define TERM_VIRTUAL    0x0002 /* Added by the optimizer.  Do not code */
+#define TERM_CODED      0x0004 /* This term is already coded */
+#define TERM_COPIED     0x0008 /* Has a child */
+#define TERM_ORINFO     0x0010 /* Need to free the WhereTerm.u.pOrInfo object */
+#define TERM_ANDINFO    0x0020 /* Need to free the WhereTerm.u.pAndInfo obj */
+#define TERM_OR_OK      0x0040 /* Used during OR-clause processing */
+#ifdef SQLITE_ENABLE_STAT4
+#  define TERM_VNULL    0x0080 /* Manufactured x>NULL or x<=NULL term */
+#else
+#  define TERM_VNULL    0x0000 /* Disabled if not using stat4 */
+#endif
+#define TERM_LIKEOPT    0x0100 /* Virtual terms from the LIKE optimization */
+#define TERM_LIKECOND   0x0200 /* Conditionally this LIKE operator term */
+#define TERM_LIKE       0x0400 /* The original LIKE operator */
+#define TERM_IS         0x0800 /* Term.pExpr is an IS operator */
 #define TERM_VARSELECT  0x1000 /* Term.pExpr contains a correlated sub-query */
 
 /*
@@ -135971,7 +139553,7 @@ struct WhereLoopBuilder {
   ExprList *pOrderBy;       /* ORDER BY clause */
   WhereLoop *pNew;          /* Template WhereLoop */
   WhereOrSet *pOrSet;       /* Record best loops here, if not NULL */
-#ifdef SQLITE_ENABLE_STAT3_OR_STAT4
+#ifdef SQLITE_ENABLE_STAT4
   UnpackedRecord *pRec;     /* Probe for stat4 (if required) */
   int nRecValid;            /* Number of valid fields currently in pRec */
 #endif
@@ -136003,6 +139585,20 @@ struct WhereLoopBuilder {
 # define SQLITE_QUERY_PLANNER_LIMIT_INCR 1000
 #endif
 
+/*
+** Each instance of this object records a change to a single node
+** in an expression tree to cause that node to point to a column
+** of an index rather than an expression or a virtual column.  All
+** such transformations need to be undone at the end of WHERE clause
+** processing.
+*/
+typedef struct WhereExprMod WhereExprMod;
+struct WhereExprMod {
+  WhereExprMod *pNext;  /* Next translation on a list of them all */
+  Expr *pExpr;          /* The Expr node that was transformed */
+  Expr orig;            /* Original value of the Expr node */
+};
+
 /*
 ** The WHERE clause processing routine has two halves.  The
 ** first part does the start of the WHERE loop and the second
@@ -136019,23 +139615,25 @@ struct WhereInfo {
   ExprList *pOrderBy;       /* The ORDER BY clause or NULL */
   ExprList *pResultSet;     /* Result set of the query */
   Expr *pWhere;             /* The complete WHERE clause */
-  LogEst iLimit;            /* LIMIT if wctrlFlags has WHERE_USE_LIMIT */
   int aiCurOnePass[2];      /* OP_OpenWrite cursors for the ONEPASS opt */
   int iContinue;            /* Jump here to continue with next record */
   int iBreak;               /* Jump here to break out of the loop */
   int savedNQueryLoop;      /* pParse->nQueryLoop outside the WHERE loop */
   u16 wctrlFlags;           /* Flags originally passed to sqlite3WhereBegin() */
+  LogEst iLimit;            /* LIMIT if wctrlFlags has WHERE_USE_LIMIT */
   u8 nLevel;                /* Number of nested loop */
   i8 nOBSat;                /* Number of ORDER BY terms satisfied by indices */
-  u8 sorted;                /* True if really sorted (not just grouped) */
   u8 eOnePass;              /* ONEPASS_OFF, or _SINGLE, or _MULTI */
-  u8 untestedTerms;         /* Not all WHERE terms resolved by outer loop */
   u8 eDistinct;             /* One of the WHERE_DISTINCT_* values */
-  u8 bOrderedInnerLoop;     /* True if only the inner-most loop is ordered */
+  unsigned bDeferredSeek :1;   /* Uses OP_DeferredSeek */
+  unsigned untestedTerms :1;   /* Not all WHERE terms resolved by outer loop */
+  unsigned bOrderedInnerLoop:1;/* True if only the inner-most loop is ordered */
+  unsigned sorted :1;          /* True if really sorted (not just grouped) */
+  LogEst nRowOut;           /* Estimated number of output rows */
   int iTop;                 /* The very beginning of the WHERE loop */
   WhereLoop *pLoops;        /* List of all WhereLoop objects */
+  WhereExprMod *pExprMods;  /* Expression modifications */
   Bitmask revMask;          /* Mask of ORDER BY terms that need reversing */
-  LogEst nRowOut;           /* Estimated number of output rows */
   WhereClause sWC;          /* Decomposition of the WHERE clause */
   WhereMaskSet sMaskSet;    /* Map cursor numbers to bitmasks */
   WhereLevel a[1];          /* Information about each nest loop in WHERE */
@@ -136049,6 +139647,8 @@ struct WhereInfo {
 SQLITE_PRIVATE Bitmask sqlite3WhereGetMask(WhereMaskSet*,int);
 #ifdef WHERETRACE_ENABLED
 SQLITE_PRIVATE void sqlite3WhereClausePrint(WhereClause *pWC);
+SQLITE_PRIVATE void sqlite3WhereTermPrint(WhereTerm *pTerm, int iTerm);
+SQLITE_PRIVATE void sqlite3WhereLoopPrint(WhereLoop *p, WhereClause *pWC);
 #endif
 SQLITE_PRIVATE WhereTerm *sqlite3WhereFindTerm(
   WhereClause *pWC,     /* The WHERE clause to be searched */
@@ -136158,6 +139758,9 @@ SQLITE_PRIVATE void sqlite3WhereTabFuncArgs(Parse*, struct SrcList_item*, WhereC
 #define WHERE_UNQ_WANTED   0x00010000  /* WHERE_ONEROW would have been helpful*/
 #define WHERE_PARTIALIDX   0x00020000  /* The automatic index is partial */
 #define WHERE_IN_EARLYOUT  0x00040000  /* Perhaps quit IN loops early */
+#define WHERE_BIGNULL_SORT 0x00080000  /* Column nEq of index is BIGNULL */
+
+#endif /* !defined(SQLITE_WHEREINT_H) */
 
 /************** End of whereInt.h ********************************************/
 /************** Continuing where we left off in wherecode.c ******************/
@@ -136460,9 +140063,9 @@ static void disableTerm(WhereLevel *pLevel, WhereTerm *pTerm){
 ** Code an OP_Affinity opcode to apply the column affinity string zAff
 ** to the n registers starting at base. 
 **
-** As an optimization, SQLITE_AFF_BLOB entries (which are no-ops) at the
-** beginning and end of zAff are ignored.  If all entries in zAff are
-** SQLITE_AFF_BLOB, then no code gets generated.
+** As an optimization, SQLITE_AFF_BLOB and SQLITE_AFF_NONE entries (which
+** are no-ops) at the beginning and end of zAff are ignored.  If all entries
+** in zAff are SQLITE_AFF_BLOB or SQLITE_AFF_NONE, then no code gets generated.
 **
 ** This routine makes its own copy of zAff so that the caller is free
 ** to modify zAff after this routine returns.
@@ -136475,15 +140078,16 @@ static void codeApplyAffinity(Parse *pParse, int base, int n, char *zAff){
   }
   assert( v!=0 );
 
-  /* Adjust base and n to skip over SQLITE_AFF_BLOB entries at the beginning
-  ** and end of the affinity string.
+  /* Adjust base and n to skip over SQLITE_AFF_BLOB and SQLITE_AFF_NONE
+  ** entries at the beginning and end of the affinity string.
   */
-  while( n>0 && zAff[0]==SQLITE_AFF_BLOB ){
+  assert( SQLITE_AFF_NONE<SQLITE_AFF_BLOB );
+  while( n>0 && zAff[0]<=SQLITE_AFF_BLOB ){
     n--;
     base++;
     zAff++;
   }
-  while( n>1 && zAff[n-1]==SQLITE_AFF_BLOB ){
+  while( n>1 && zAff[n-1]<=SQLITE_AFF_BLOB ){
     n--;
   }
 
@@ -136556,7 +140160,8 @@ static Expr *removeUnindexableInClauseTerms(
   Expr *pX              /* The IN expression to be reduced */
 ){
   sqlite3 *db = pParse->db;
-  Expr *pNew = sqlite3ExprDup(db, pX, 0);
+  Expr *pNew;
+  pNew = sqlite3ExprDup(db, pX, 0);
   if( db->mallocFailed==0 ){
     ExprList *pOrigRhs = pNew->x.pSelect->pEList;  /* Original unmodified RHS */
     ExprList *pOrigLhs = pNew->pLeft->x.pList;     /* Original unmodified LHS */
@@ -136733,7 +140338,7 @@ static int codeEqualityTerm(
           if( i==iEq ){
             pIn->iCur = iTab;
             pIn->eEndLoopOp = bRev ? OP_Prev : OP_Next;
-            if( iEq>0 && (pLoop->wsFlags & WHERE_VIRTUALTABLE)==0 ){
+            if( iEq>0 ){
               pIn->iBase = iReg - i;
               pIn->nPrefix = i;
               pLoop->wsFlags |= WHERE_IN_EARLYOUT;
@@ -136964,7 +140569,7 @@ static int codeCursorHintCheckExpr(Walker *pWalker, Expr *pExpr){
   assert( pHint->pIdx!=0 );
   if( pExpr->op==TK_COLUMN
    && pExpr->iTable==pHint->iTabCur
-   && sqlite3ColumnOfIndex(pHint->pIdx, pExpr->iColumn)<0
+   && sqlite3TableColumnToIndex(pHint->pIdx, pExpr->iColumn)<0
   ){
     pWalker->eCode = 1;
   }
@@ -137032,7 +140637,7 @@ static int codeCursorHintFixExpr(Walker *pWalker, Expr *pExpr){
       pExpr->iTable = reg;
     }else if( pHint->pIdx!=0 ){
       pExpr->iTable = pHint->iIdxCur;
-      pExpr->iColumn = sqlite3ColumnOfIndex(pHint->pIdx, pExpr->iColumn);
+      pExpr->iColumn = sqlite3TableColumnToIndex(pHint->pIdx, pExpr->iColumn);
       assert( pExpr->iColumn>=0 );
     }
   }else if( pExpr->op==TK_AGG_FUNCTION ){
@@ -137141,7 +140746,7 @@ static void codeCursorHint(
     }
 
     /* If we survive all prior tests, that means this term is worth hinting */
-    pExpr = sqlite3ExprAnd(db, pExpr, sqlite3ExprDup(db, pTerm->pExpr, 0));
+    pExpr = sqlite3ExprAnd(pParse, pExpr, sqlite3ExprDup(db, pTerm->pExpr, 0));
   }
   if( pExpr!=0 ){
     sWalker.xExprCallback = codeCursorHintFixExpr;
@@ -137185,6 +140790,7 @@ static void codeDeferredSeek(
   assert( iIdxCur>0 );
   assert( pIdx->aiColumn[pIdx->nColumn-1]==-1 );
   
+  pWInfo->bDeferredSeek = 1;
   sqlite3VdbeAddOp3(v, OP_DeferredSeek, iIdxCur, 0, iCur);
   if( (pWInfo->wctrlFlags & WHERE_OR_SUBCLAUSE)
    && DbMaskAllZero(sqlite3ParseToplevel(pParse)->writeMask)
@@ -137195,8 +140801,12 @@ static void codeDeferredSeek(
     if( ai ){
       ai[0] = pTab->nCol;
       for(i=0; i<pIdx->nColumn-1; i++){
+        int x1, x2;
         assert( pIdx->aiColumn[i]<pTab->nCol );
-        if( pIdx->aiColumn[i]>=0 ) ai[pIdx->aiColumn[i]+1] = i+1;
+        x1 = pIdx->aiColumn[i];
+        x2 = sqlite3TableColumnToStorage(pTab, x1);
+        testcase( x1!=x2 );
+        if( x1>=0 ) ai[x2+1] = i+1;
       }
       sqlite3VdbeChangeP4(v, -1, (char*)ai, P4_INTARRAY);
     }
@@ -137247,8 +140857,24 @@ typedef struct IdxExprTrans {
   int iTabCur;       /* The cursor of the corresponding table */
   int iIdxCur;       /* The cursor for the index */
   int iIdxCol;       /* The column for the index */
+  int iTabCol;       /* The column for the table */
+  WhereInfo *pWInfo; /* Complete WHERE clause information */
+  sqlite3 *db;       /* Database connection (for malloc()) */
 } IdxExprTrans;
 
+/*
+** Preserve pExpr on the WhereETrans list of the WhereInfo.
+*/
+static void preserveExpr(IdxExprTrans *pTrans, Expr *pExpr){
+  WhereExprMod *pNew;
+  pNew = sqlite3DbMallocRaw(pTrans->db, sizeof(*pNew));
+  if( pNew==0 ) return;
+  pNew->pNext = pTrans->pWInfo->pExprMods;
+  pTrans->pWInfo->pExprMods = pNew;
+  pNew->pExpr = pExpr;
+  memcpy(&pNew->orig, pExpr, sizeof(*pExpr));
+}
+
 /* The walker node callback used to transform matching expressions into
 ** a reference to an index column for an index on an expression.
 **
@@ -137258,20 +140884,49 @@ typedef struct IdxExprTrans {
 static int whereIndexExprTransNode(Walker *p, Expr *pExpr){
   IdxExprTrans *pX = p->u.pIdxTrans;
   if( sqlite3ExprCompare(0, pExpr, pX->pIdxExpr, pX->iTabCur)==0 ){
+    preserveExpr(pX, pExpr);
+    pExpr->affExpr = sqlite3ExprAffinity(pExpr);
     pExpr->op = TK_COLUMN;
     pExpr->iTable = pX->iIdxCur;
     pExpr->iColumn = pX->iIdxCol;
     pExpr->y.pTab = 0;
+    testcase( ExprHasProperty(pExpr, EP_Skip) );
+    testcase( ExprHasProperty(pExpr, EP_Unlikely) );
+    ExprClearProperty(pExpr, EP_Skip|EP_Unlikely);
     return WRC_Prune;
   }else{
     return WRC_Continue;
   }
 }
 
+#ifndef SQLITE_OMIT_GENERATED_COLUMNS
+/* A walker node callback that translates a column reference to a table
+** into a corresponding column reference of an index.
+*/
+static int whereIndexExprTransColumn(Walker *p, Expr *pExpr){
+  if( pExpr->op==TK_COLUMN ){
+    IdxExprTrans *pX = p->u.pIdxTrans;
+    if( pExpr->iTable==pX->iTabCur && pExpr->iColumn==pX->iTabCol ){
+      assert( pExpr->y.pTab!=0 );
+      preserveExpr(pX, pExpr);
+      pExpr->affExpr = sqlite3TableColumnAffinity(pExpr->y.pTab,pExpr->iColumn);
+      pExpr->iTable = pX->iIdxCur;
+      pExpr->iColumn = pX->iIdxCol;
+      pExpr->y.pTab = 0;
+    }
+  }
+  return WRC_Continue;
+}
+#endif /* SQLITE_OMIT_GENERATED_COLUMNS */
+
 /*
 ** For an indexes on expression X, locate every instance of expression X
 ** in pExpr and change that subexpression into a reference to the appropriate
 ** column of the index.
+**
+** 2019-10-24: Updated to also translate references to a VIRTUAL column in
+** the table into references to the corresponding (stored) column of the
+** index.
 */
 static void whereIndexExprTrans(
   Index *pIdx,      /* The Index */
@@ -137281,20 +140936,48 @@ static void whereIndexExprTrans(
 ){
   int iIdxCol;               /* Column number of the index */
   ExprList *aColExpr;        /* Expressions that are indexed */
+  Table *pTab;
   Walker w;
   IdxExprTrans x;
   aColExpr = pIdx->aColExpr;
-  if( aColExpr==0 ) return;  /* Not an index on expressions */
+  if( aColExpr==0 && !pIdx->bHasVCol ){
+    /* The index does not reference any expressions or virtual columns
+    ** so no translations are needed. */
+    return;
+  }
+  pTab = pIdx->pTable;
   memset(&w, 0, sizeof(w));
-  w.xExprCallback = whereIndexExprTransNode;
   w.u.pIdxTrans = &x;
   x.iTabCur = iTabCur;
   x.iIdxCur = iIdxCur;
-  for(iIdxCol=0; iIdxCol<aColExpr->nExpr; iIdxCol++){
-    if( pIdx->aiColumn[iIdxCol]!=XN_EXPR ) continue;
-    assert( aColExpr->a[iIdxCol].pExpr!=0 );
+  x.pWInfo = pWInfo;
+  x.db = pWInfo->pParse->db;
+  for(iIdxCol=0; iIdxCol<pIdx->nColumn; iIdxCol++){
+    i16 iRef = pIdx->aiColumn[iIdxCol];
+    if( iRef==XN_EXPR ){
+      assert( aColExpr->a[iIdxCol].pExpr!=0 );
+      x.pIdxExpr = aColExpr->a[iIdxCol].pExpr;
+      if( sqlite3ExprIsConstant(x.pIdxExpr) ) continue;
+      w.xExprCallback = whereIndexExprTransNode;
+#ifndef SQLITE_OMIT_GENERATED_COLUMNS
+    }else if( iRef>=0
+       && (pTab->aCol[iRef].colFlags & COLFLAG_VIRTUAL)!=0
+       && (pTab->aCol[iRef].zColl==0
+           || sqlite3StrICmp(pTab->aCol[iRef].zColl, sqlite3StrBINARY)==0)
+    ){
+      /* Check to see if there are direct references to generated columns
+      ** that are contained in the index.  Pulling the generated column
+      ** out of the index is an optimization only - the main table is always
+      ** available if the index cannot be used.  To avoid unnecessary
+      ** complication, omit this optimization if the collating sequence for
+      ** the column is non-standard */
+      x.iTabCol = iRef;
+      w.xExprCallback = whereIndexExprTransColumn;
+#endif /* SQLITE_OMIT_GENERATED_COLUMNS */
+    }else{
+      continue;
+    }
     x.iIdxCol = iIdxCol;
-    x.pIdxExpr = aColExpr->a[iIdxCol].pExpr;
     sqlite3WalkExpr(&w, pWInfo->pWhere);
     sqlite3WalkExprList(&w, pWInfo->pOrderBy);
     sqlite3WalkExprList(&w, pWInfo->pResultSet);
@@ -137366,6 +141049,21 @@ SQLITE_PRIVATE Bitmask sqlite3WhereCodeOneLoopStart(
   pLevel->notReady = notReady & ~sqlite3WhereGetMask(&pWInfo->sMaskSet, iCur);
   bRev = (pWInfo->revMask>>iLevel)&1;
   VdbeModuleComment((v, "Begin WHERE-loop%d: %s",iLevel,pTabItem->pTab->zName));
+#if WHERETRACE_ENABLED /* 0x20800 */
+  if( sqlite3WhereTrace & 0x800 ){
+    sqlite3DebugPrintf("Coding level %d of %d:  notReady=%llx  iFrom=%d\n",
+       iLevel, pWInfo->nLevel, (u64)notReady, pLevel->iFrom);
+    sqlite3WhereLoopPrint(pLoop, pWC);
+  }
+  if( sqlite3WhereTrace & 0x20000 ){
+    if( iLevel==0 ){
+      sqlite3DebugPrintf("WHERE clause being coded:\n");
+      sqlite3TreeViewExpr(0, pWInfo->pWhere, 0);
+    }
+    sqlite3DebugPrintf("All WHERE-clause terms before coding:\n");
+    sqlite3WhereClausePrint(pWC);
+  }
+#endif
 
   /* Create labels for the "break" and "continue" instructions
   ** for the current loop.  Jump to addrBrk to break out of a loop.
@@ -137445,9 +141143,12 @@ SQLITE_PRIVATE Bitmask sqlite3WhereCodeOneLoopStart(
     iIn = pLevel->u.in.nIn;
     for(j=nConstraint-1; j>=0; j--){
       pTerm = pLoop->aLTerm[j];
+      if( (pTerm->eOperator & WO_IN)!=0 ) iIn--;
       if( j<16 && (pLoop->u.vtab.omitMask>>j)&1 ){
         disableTerm(pLevel, pTerm);
-      }else if( (pTerm->eOperator & WO_IN)!=0 ){
+      }else if( (pTerm->eOperator & WO_IN)!=0
+        && sqlite3ExprVectorSize(pTerm->pExpr->pLeft)==1
+      ){
         Expr *pCompare;  /* The comparison operator */
         Expr *pRight;    /* RHS of the comparison */
         VdbeOp *pOp;     /* Opcode to access the value of the IN constraint */
@@ -137458,8 +141159,8 @@ SQLITE_PRIVATE Bitmask sqlite3WhereCodeOneLoopStart(
         ** encoding of the value in the register, so it *must* be reloaded. */
         assert( pLevel->u.in.aInLoop!=0 || db->mallocFailed );
         if( !db->mallocFailed ){
-          assert( iIn>0 );
-          pOp = sqlite3VdbeGetOp(v, pLevel->u.in.aInLoop[--iIn].addrInTop);
+          assert( iIn>=0 && iIn<pLevel->u.in.nIn );
+          pOp = sqlite3VdbeGetOp(v, pLevel->u.in.aInLoop[iIn].addrInTop);
           assert( pOp->opcode==OP_Column || pOp->opcode==OP_Rowid );
           assert( pOp->opcode!=OP_Column || pOp->p3==iReg+j+2 );
           assert( pOp->opcode!=OP_Rowid || pOp->p2==iReg+j+2 );
@@ -137483,6 +141184,7 @@ SQLITE_PRIVATE Bitmask sqlite3WhereCodeOneLoopStart(
         }
       }
     }
+    assert( iIn==0 || db->mallocFailed );
     /* These registers need to be preserved in case there is an IN operator
     ** loop.  So we could deallocate the registers here (and potentially
     ** reuse them later) if (pLoop->wsFlags & WHERE_IN_ABLE)==0.  But it seems
@@ -137690,32 +141392,12 @@ SQLITE_PRIVATE Bitmask sqlite3WhereCodeOneLoopStart(
     u8 bSeekPastNull = 0;        /* True to seek past initial nulls */
     u8 bStopAtNull = 0;          /* Add condition to terminate at NULLs */
     int omitTable;               /* True if we use the index only */
-
+    int regBignull = 0;          /* big-null flag register */
 
     pIdx = pLoop->u.btree.pIndex;
     iIdxCur = pLevel->iIdxCur;
     assert( nEq>=pLoop->nSkip );
 
-    /* If this loop satisfies a sort order (pOrderBy) request that 
-    ** was passed to this function to implement a "SELECT min(x) ..." 
-    ** query, then the caller will only allow the loop to run for
-    ** a single iteration. This means that the first row returned
-    ** should not have a NULL value stored in 'x'. If column 'x' is
-    ** the first one after the nEq equality constraints in the index,
-    ** this requires some special handling.
-    */
-    assert( pWInfo->pOrderBy==0
-         || pWInfo->pOrderBy->nExpr==1
-         || (pWInfo->wctrlFlags&WHERE_ORDERBY_MIN)==0 );
-    if( (pWInfo->wctrlFlags&WHERE_ORDERBY_MIN)!=0
-     && pWInfo->nOBSat>0
-     && (pIdx->nKeyCol>nEq)
-    ){
-      assert( pLoop->nSkip==0 );
-      bSeekPastNull = 1;
-      nExtraReg = 1;
-    }
-
     /* Find any inequality constraint terms for the start and end 
     ** of the range. 
     */
@@ -137756,6 +141438,25 @@ SQLITE_PRIVATE Bitmask sqlite3WhereCodeOneLoopStart(
     }
     assert( pRangeEnd==0 || (pRangeEnd->wtFlags & TERM_VNULL)==0 );
 
+    /* If the WHERE_BIGNULL_SORT flag is set, then index column nEq uses
+    ** a non-default "big-null" sort (either ASC NULLS LAST or DESC NULLS 
+    ** FIRST). In both cases separate ordered scans are made of those
+    ** index entries for which the column is null and for those for which
+    ** it is not. For an ASC sort, the non-NULL entries are scanned first.
+    ** For DESC, NULL entries are scanned first.
+    */
+    if( (pLoop->wsFlags & (WHERE_TOP_LIMIT|WHERE_BTM_LIMIT))==0
+     && (pLoop->wsFlags & WHERE_BIGNULL_SORT)!=0
+    ){
+      assert( bSeekPastNull==0 && nExtraReg==0 && nBtm==0 && nTop==0 );
+      assert( pRangeEnd==0 && pRangeStart==0 );
+      testcase( pLoop->nSkip>0 );
+      nExtraReg = 1;
+      bSeekPastNull = 1;
+      pLevel->regBignull = regBignull = ++pParse->nMem;
+      pLevel->addrBignull = sqlite3VdbeMakeLabel(pParse);
+    }
+
     /* If we are doing a reverse order scan on an ascending index, or
     ** a forward order scan on a descending index, interchange the 
     ** start and end terms (pRangeStart and pRangeEnd).
@@ -137778,7 +141479,7 @@ SQLITE_PRIVATE Bitmask sqlite3WhereCodeOneLoopStart(
     if( zStartAff && nTop ){
       zEndAff = sqlite3DbStrDup(db, &zStartAff[nEq]);
     }
-    addrNxt = pLevel->addrNxt;
+    addrNxt = (regBignull ? pLevel->addrBignull : pLevel->addrNxt);
 
     testcase( pRangeStart && (pRangeStart->eOperator & WO_LE)!=0 );
     testcase( pRangeStart && (pRangeStart->eOperator & WO_GE)!=0 );
@@ -137812,10 +141513,14 @@ SQLITE_PRIVATE Bitmask sqlite3WhereCodeOneLoopStart(
       }
       bSeekPastNull = 0;
     }else if( bSeekPastNull ){
+      startEq = 0;
       sqlite3VdbeAddOp2(v, OP_Null, 0, regBase+nEq);
+      start_constraints = 1;
       nConstraint++;
-      startEq = 0;
+    }else if( regBignull ){
+      sqlite3VdbeAddOp2(v, OP_Null, 0, regBase+nEq);
       start_constraints = 1;
+      nConstraint++;
     }
     codeApplyAffinity(pParse, regBase, nConstraint - bSeekPastNull, zStartAff);
     if( pLoop->nSkip>0 && nConstraint==pLoop->nSkip ){
@@ -137826,6 +141531,11 @@ SQLITE_PRIVATE Bitmask sqlite3WhereCodeOneLoopStart(
       if( pLoop->wsFlags & WHERE_IN_EARLYOUT ){
         sqlite3VdbeAddOp1(v, OP_SeekHit, iIdxCur);
       }
+      if( regBignull ){
+        sqlite3VdbeAddOp2(v, OP_Integer, 1, regBignull);
+        VdbeComment((v, "NULL-scan pass ctr"));
+      }
+
       op = aStartOp[(start_constraints<<2) + (startEq<<1) + bRev];
       assert( op!=0 );
       sqlite3VdbeAddOp4Int(v, op, iIdxCur, addrNxt, regBase, nConstraint);
@@ -137836,6 +141546,23 @@ SQLITE_PRIVATE Bitmask sqlite3WhereCodeOneLoopStart(
       VdbeCoverageIf(v, op==OP_SeekGE);  testcase( op==OP_SeekGE );
       VdbeCoverageIf(v, op==OP_SeekLE);  testcase( op==OP_SeekLE );
       VdbeCoverageIf(v, op==OP_SeekLT);  testcase( op==OP_SeekLT );
+
+      assert( bSeekPastNull==0 || bStopAtNull==0 );
+      if( regBignull ){
+        assert( bSeekPastNull==1 || bStopAtNull==1 );
+        assert( bSeekPastNull==!bStopAtNull );
+        assert( bStopAtNull==startEq );
+        sqlite3VdbeAddOp2(v, OP_Goto, 0, sqlite3VdbeCurrentAddr(v)+2);
+        op = aStartOp[(nConstraint>1)*4 + 2 + bRev];
+        sqlite3VdbeAddOp4Int(v, op, iIdxCur, addrNxt, regBase, 
+                             nConstraint-startEq);
+        VdbeCoverage(v);
+        VdbeCoverageIf(v, op==OP_Rewind);  testcase( op==OP_Rewind );
+        VdbeCoverageIf(v, op==OP_Last);    testcase( op==OP_Last );
+        VdbeCoverageIf(v, op==OP_SeekGE);  testcase( op==OP_SeekGE );
+        VdbeCoverageIf(v, op==OP_SeekLE);  testcase( op==OP_SeekLE );
+        assert( op==OP_Rewind || op==OP_Last || op==OP_SeekGE || op==OP_SeekLE);
+      }
     }
 
     /* Load the value for the inequality constraint at the end of the
@@ -137867,8 +141594,10 @@ SQLITE_PRIVATE Bitmask sqlite3WhereCodeOneLoopStart(
         endEq = 1;
       }
     }else if( bStopAtNull ){
-      sqlite3VdbeAddOp2(v, OP_Null, 0, regBase+nEq);
-      endEq = 0;
+      if( regBignull==0 ){
+        sqlite3VdbeAddOp2(v, OP_Null, 0, regBase+nEq);
+        endEq = 0;
+      }
       nConstraint++;
     }
     sqlite3DbFree(db, zStartAff);
@@ -137879,6 +141608,12 @@ SQLITE_PRIVATE Bitmask sqlite3WhereCodeOneLoopStart(
 
     /* Check if the index cursor is past the end of the range. */
     if( nConstraint ){
+      if( regBignull ){
+        /* Except, skip the end-of-range check while doing the NULL-scan */
+        sqlite3VdbeAddOp2(v, OP_IfNot, regBignull, sqlite3VdbeCurrentAddr(v)+3);
+        VdbeComment((v, "If NULL-scan 2nd pass"));
+        VdbeCoverage(v);
+      }
       op = aEndOp[bRev*2 + endEq];
       sqlite3VdbeAddOp4Int(v, op, iIdxCur, addrNxt, regBase, nConstraint);
       testcase( op==OP_IdxGT );  VdbeCoverageIf(v, op==OP_IdxGT );
@@ -137886,6 +141621,23 @@ SQLITE_PRIVATE Bitmask sqlite3WhereCodeOneLoopStart(
       testcase( op==OP_IdxLT );  VdbeCoverageIf(v, op==OP_IdxLT );
       testcase( op==OP_IdxLE );  VdbeCoverageIf(v, op==OP_IdxLE );
     }
+    if( regBignull ){
+      /* During a NULL-scan, check to see if we have reached the end of
+      ** the NULLs */
+      assert( bSeekPastNull==!bStopAtNull );
+      assert( bSeekPastNull+bStopAtNull==1 );
+      assert( nConstraint+bSeekPastNull>0 );
+      sqlite3VdbeAddOp2(v, OP_If, regBignull, sqlite3VdbeCurrentAddr(v)+2);
+      VdbeComment((v, "If NULL-scan 1st pass"));
+      VdbeCoverage(v);
+      op = aEndOp[bRev*2 + bSeekPastNull];
+      sqlite3VdbeAddOp4Int(v, op, iIdxCur, addrNxt, regBase,
+                           nConstraint+bSeekPastNull);
+      testcase( op==OP_IdxGT );  VdbeCoverageIf(v, op==OP_IdxGT );
+      testcase( op==OP_IdxGE );  VdbeCoverageIf(v, op==OP_IdxGE );
+      testcase( op==OP_IdxLT );  VdbeCoverageIf(v, op==OP_IdxLT );
+      testcase( op==OP_IdxLE );  VdbeCoverageIf(v, op==OP_IdxLE );
+    }
 
     if( pLoop->wsFlags & WHERE_IN_EARLYOUT ){
       sqlite3VdbeAddOp2(v, OP_SeekHit, iIdxCur, 1);
@@ -137897,10 +141649,10 @@ SQLITE_PRIVATE Bitmask sqlite3WhereCodeOneLoopStart(
     if( omitTable ){
       /* pIdx is a covering index.  No need to access the main table. */
     }else if( HasRowid(pIdx->pTable) ){
-      if( (pWInfo->wctrlFlags & WHERE_SEEK_TABLE) || (
-          (pWInfo->wctrlFlags & WHERE_SEEK_UNIQ_TABLE) 
-       && (pWInfo->eOnePass==ONEPASS_SINGLE)
-      )){
+      if( (pWInfo->wctrlFlags & WHERE_SEEK_TABLE)
+       || ( (pWInfo->wctrlFlags & WHERE_SEEK_UNIQ_TABLE)!=0
+           && (pWInfo->eOnePass==ONEPASS_SINGLE || pLoop->nLTerm==0) )
+      ){
         iRowidReg = ++pParse->nMem;
         sqlite3VdbeAddOp2(v, OP_IdxRowid, iIdxCur, iRowidReg);
         sqlite3VdbeAddOp3(v, OP_NotExists, iCur, 0, iRowidReg);
@@ -137912,40 +141664,53 @@ SQLITE_PRIVATE Bitmask sqlite3WhereCodeOneLoopStart(
       Index *pPk = sqlite3PrimaryKeyIndex(pIdx->pTable);
       iRowidReg = sqlite3GetTempRange(pParse, pPk->nKeyCol);
       for(j=0; j<pPk->nKeyCol; j++){
-        k = sqlite3ColumnOfIndex(pIdx, pPk->aiColumn[j]);
+        k = sqlite3TableColumnToIndex(pIdx, pPk->aiColumn[j]);
         sqlite3VdbeAddOp3(v, OP_Column, iIdxCur, k, iRowidReg+j);
       }
       sqlite3VdbeAddOp4Int(v, OP_NotFound, iCur, addrCont,
                            iRowidReg, pPk->nKeyCol); VdbeCoverage(v);
     }
 
-    /* If pIdx is an index on one or more expressions, then look through
-    ** all the expressions in pWInfo and try to transform matching expressions
-    ** into reference to index columns.
-    **
-    ** Do not do this for the RHS of a LEFT JOIN. This is because the 
-    ** expression may be evaluated after OP_NullRow has been executed on
-    ** the cursor. In this case it is important to do the full evaluation,
-    ** as the result of the expression may not be NULL, even if all table
-    ** column values are.  https://www.sqlite.org/src/info/7fa8049685b50b5a
-    **
-    ** Also, do not do this when processing one index an a multi-index
-    ** OR clause, since the transformation will become invalid once we
-    ** move forward to the next index.
-    ** https://sqlite.org/src/info/4e8e4857d32d401f
-    */
-    if( pLevel->iLeftJoin==0 && (pWInfo->wctrlFlags & WHERE_OR_SUBCLAUSE)==0 ){
-      whereIndexExprTrans(pIdx, iCur, iIdxCur, pWInfo);
-    }
-
-    /* If a partial index is driving the loop, try to eliminate WHERE clause
-    ** terms from the query that must be true due to the WHERE clause of
-    ** the partial index
-    */
-    if( pIdx->pPartIdxWhere ){
-      whereApplyPartialIndexConstraints(pIdx->pPartIdxWhere, iCur, pWC);
+    if( pLevel->iLeftJoin==0 ){
+      /* If pIdx is an index on one or more expressions, then look through
+      ** all the expressions in pWInfo and try to transform matching expressions
+      ** into reference to index columns.  Also attempt to translate references
+      ** to virtual columns in the table into references to (stored) columns
+      ** of the index.
+      **
+      ** Do not do this for the RHS of a LEFT JOIN. This is because the 
+      ** expression may be evaluated after OP_NullRow has been executed on
+      ** the cursor. In this case it is important to do the full evaluation,
+      ** as the result of the expression may not be NULL, even if all table
+      ** column values are.  https://www.sqlite.org/src/info/7fa8049685b50b5a
+      **
+      ** Also, do not do this when processing one index an a multi-index
+      ** OR clause, since the transformation will become invalid once we
+      ** move forward to the next index.
+      ** https://sqlite.org/src/info/4e8e4857d32d401f
+      */
+      if( (pWInfo->wctrlFlags & WHERE_OR_SUBCLAUSE)==0 ){
+        whereIndexExprTrans(pIdx, iCur, iIdxCur, pWInfo);
+      }
+  
+      /* If a partial index is driving the loop, try to eliminate WHERE clause
+      ** terms from the query that must be true due to the WHERE clause of
+      ** the partial index.
+      **
+      ** 2019-11-02 ticket 623eff57e76d45f6: This optimization does not work
+      ** for a LEFT JOIN.
+      */
+      if( pIdx->pPartIdxWhere ){
+        whereApplyPartialIndexConstraints(pIdx->pPartIdxWhere, iCur, pWC);
+      }
+    }else{
+      testcase( pIdx->pPartIdxWhere );
+      /* The following assert() is not a requirement, merely an observation:
+      ** The OR-optimization doesn't work for the right hand table of
+      ** a LEFT JOIN: */
+      assert( (pWInfo->wctrlFlags & WHERE_OR_SUBCLAUSE)==0 );
     }
-
+  
     /* Record the instruction used to terminate the loop. */
     if( pLoop->wsFlags & WHERE_ONEROW ){
       pLevel->op = OP_Noop;
@@ -138106,7 +141871,7 @@ SQLITE_PRIVATE Bitmask sqlite3WhereCodeOneLoopStart(
         if( (pWC->a[iTerm].eOperator & WO_ALL)==0 ) continue;
         testcase( pWC->a[iTerm].wtFlags & TERM_ORINFO );
         pExpr = sqlite3ExprDup(db, pExpr, 0);
-        pAndExpr = sqlite3ExprAnd(db, pAndExpr, pExpr);
+        pAndExpr = sqlite3ExprAnd(pParse, pAndExpr, pExpr);
       }
       if( pAndExpr ){
         /* The extra 0x10000 bit on the opcode is masked off and does not
@@ -138130,9 +141895,9 @@ SQLITE_PRIVATE Bitmask sqlite3WhereCodeOneLoopStart(
         WhereInfo *pSubWInfo;           /* Info for single OR-term scan */
         Expr *pOrExpr = pOrTerm->pExpr; /* Current OR clause term */
         int jmp1 = 0;                   /* Address of jump operation */
-        assert( (pTabItem[0].fg.jointype & JT_LEFT)==0 
-             || ExprHasProperty(pOrExpr, EP_FromJoin) 
-        );
+        testcase( (pTabItem[0].fg.jointype & JT_LEFT)!=0
+               && !ExprHasProperty(pOrExpr, EP_FromJoin)
+        ); /* See TH3 vtab25.400 and ticket 614b25314c766238 */
         if( pAndExpr ){
           pAndExpr->pLeft = pOrExpr;
           pOrExpr = pAndExpr;
@@ -138172,7 +141937,7 @@ SQLITE_PRIVATE Bitmask sqlite3WhereCodeOneLoopStart(
               r = sqlite3GetTempRange(pParse, nPk);
               for(iPk=0; iPk<nPk; iPk++){
                 int iCol = pPk->aiColumn[iPk];
-                sqlite3ExprCodeGetColumnOfTable(v, pTab, iCur, iCol, r+iPk);
+                sqlite3ExprCodeGetColumnOfTable(v, pTab, iCur, iCol,r+iPk);
               }
 
               /* Check if the temp table already contains this key. If so,
@@ -138257,7 +142022,7 @@ SQLITE_PRIVATE Bitmask sqlite3WhereCodeOneLoopStart(
     sqlite3VdbeGoto(v, pLevel->addrBrk);
     sqlite3VdbeResolveLabel(v, iLoopBody);
 
-    if( pWInfo->nLevel>1 ) sqlite3StackFree(db, pOrTab);
+    if( pWInfo->nLevel>1 ){ sqlite3StackFree(db, pOrTab); }
     if( !untestedTerms ) disableTerm(pLevel, pTerm);
   }else
 #endif /* SQLITE_OMIT_OR_OPTIMIZATION */
@@ -138354,6 +142119,10 @@ SQLITE_PRIVATE Bitmask sqlite3WhereCodeOneLoopStart(
         VdbeNoopComment((v, "WhereTerm[%d] (%p) priority=%d",
                          pWC->nTerm-j, pTerm, iLoop));
       }
+      if( sqlite3WhereTrace & 0x800 ){
+        sqlite3DebugPrintf("Coding auxiliary constraint:\n");
+        sqlite3WhereTermPrint(pTerm, pWC->nTerm-j);
+      }
 #endif
       sqlite3ExprIfFalse(pParse, pE, addrCont, SQLITE_JUMPIFNULL);
       if( skipLikeAddr ) sqlite3VdbeJumpHere(v, skipLikeAddr);
@@ -138377,8 +142146,14 @@ SQLITE_PRIVATE Bitmask sqlite3WhereCodeOneLoopStart(
     if( (pTerm->eOperator & (WO_EQ|WO_IS))==0 ) continue;
     if( (pTerm->eOperator & WO_EQUIV)==0 ) continue;
     if( pTerm->leftCursor!=iCur ) continue;
-    if( pLevel->iLeftJoin ) continue;
+    if( pTabItem->fg.jointype & JT_LEFT ) continue;
     pE = pTerm->pExpr;
+#ifdef WHERETRACE_ENABLED /* 0x800 */
+    if( sqlite3WhereTrace & 0x800 ){
+      sqlite3DebugPrintf("Coding transitive constraint:\n");
+      sqlite3WhereTermPrint(pTerm, pWC->nTerm-j);
+    }
+#endif
     assert( !ExprHasProperty(pE, EP_FromJoin) );
     assert( (pTerm->prereqRight & pLevel->notReady)!=0 );
     pAlt = sqlite3WhereFindTerm(pWC, iCur, pTerm->u.leftColumn, notReady,
@@ -138421,6 +142196,17 @@ SQLITE_PRIVATE Bitmask sqlite3WhereCodeOneLoopStart(
     }
   }
 
+#if WHERETRACE_ENABLED /* 0x20800 */
+  if( sqlite3WhereTrace & 0x20000 ){
+    sqlite3DebugPrintf("All WHERE-clause terms after coding level %d:\n",
+                       iLevel);
+    sqlite3WhereClausePrint(pWC);
+  }
+  if( sqlite3WhereTrace & 0x800 ){
+    sqlite3DebugPrintf("End Coding level %d:  notReady=%llx\n",
+       iLevel, (u64)pLevel->notReady);
+  }
+#endif
   return pLevel->notReady;
 }
 
@@ -138512,7 +142298,7 @@ static int whereClauseInsert(WhereClause *pWC, Expr *p, u16 wtFlags){
   }else{
     pTerm->truthProb = 1;
   }
-  pTerm->pExpr = sqlite3ExprSkipCollate(p);
+  pTerm->pExpr = sqlite3ExprSkipCollateAndLikely(p);
   pTerm->wtFlags = wtFlags;
   pTerm->pWC = pWC;
   pTerm->iParent = -1;
@@ -138537,31 +142323,14 @@ static int allowedOp(int op){
 /*
 ** Commute a comparison operator.  Expressions of the form "X op Y"
 ** are converted into "Y op X".
-**
-** If left/right precedence rules come into play when determining the
-** collating sequence, then COLLATE operators are adjusted to ensure
-** that the collating sequence does not change.  For example:
-** "Y collate NOCASE op X" becomes "X op Y" because any collation sequence on
-** the left hand side of a comparison overrides any collation sequence 
-** attached to the right. For the same reason the EP_Collate flag
-** is not commuted.
-*/
-static void exprCommute(Parse *pParse, Expr *pExpr){
-  u16 expRight = (pExpr->pRight->flags & EP_Collate);
-  u16 expLeft = (pExpr->pLeft->flags & EP_Collate);
-  assert( allowedOp(pExpr->op) && pExpr->op!=TK_IN );
-  if( expRight==expLeft ){
-    /* Either X and Y both have COLLATE operator or neither do */
-    if( expRight ){
-      /* Both X and Y have COLLATE operators.  Make sure X is always
-      ** used by clearing the EP_Collate flag from Y. */
-      pExpr->pRight->flags &= ~EP_Collate;
-    }else if( sqlite3ExprCollSeq(pParse, pExpr->pLeft)!=0 ){
-      /* Neither X nor Y have COLLATE operators, but X has a non-default
-      ** collating sequence.  So add the EP_Collate marker on X to cause
-      ** it to be searched first. */
-      pExpr->pLeft->flags |= EP_Collate;
-    }
+*/
+static u16 exprCommute(Parse *pParse, Expr *pExpr){
+  if( pExpr->pLeft->op==TK_VECTOR
+   || pExpr->pRight->op==TK_VECTOR
+   || sqlite3BinaryCompareCollSeq(pParse, pExpr->pLeft, pExpr->pRight) !=
+      sqlite3BinaryCompareCollSeq(pParse, pExpr->pRight, pExpr->pLeft)
+  ){
+    pExpr->flags ^= EP_Commuted;
   }
   SWAP(Expr*,pExpr->pRight,pExpr->pLeft);
   if( pExpr->op>=TK_GT ){
@@ -138572,6 +142341,7 @@ static void exprCommute(Parse *pParse, Expr *pExpr){
     assert( pExpr->op>=TK_GT && pExpr->op<=TK_GE );
     pExpr->op = ((pExpr->op-TK_GT)^2)+TK_GT;
   }
+  return 0;
 }
 
 /*
@@ -138690,27 +142460,38 @@ static int isLikeOrGlob(
           zNew[iTo++] = zNew[iFrom];
         }
         zNew[iTo] = 0;
+        assert( iTo>0 );
 
-        /* If the RHS begins with a digit or a minus sign, then the LHS must be
-        ** an ordinary column (not a virtual table column) with TEXT affinity.
-        ** Otherwise the LHS might be numeric and "lhs >= rhs" would be false
-        ** even though "lhs LIKE rhs" is true.  But if the RHS does not start
-        ** with a digit or '-', then "lhs LIKE rhs" will always be false if
-        ** the LHS is numeric and so the optimization still works.
+        /* If the LHS is not an ordinary column with TEXT affinity, then the
+        ** pattern prefix boundaries (both the start and end boundaries) must
+        ** not look like a number.  Otherwise the pattern might be treated as
+        ** a number, which will invalidate the LIKE optimization.
         **
-        ** 2018-09-10 ticket c94369cae9b561b1f996d0054bfab11389f9d033
-        ** The RHS pattern must not be '/%' because the termination condition
-        ** will then become "x<'0'" and if the affinity is numeric, will then
-        ** be converted into "x<0", which is incorrect.
+        ** Getting this right has been a persistent source of bugs in the
+        ** LIKE optimization.  See, for example:
+        **    2018-09-10 https://sqlite.org/src/info/c94369cae9b561b1
+        **    2019-05-02 https://sqlite.org/src/info/b043a54c3de54b28
+        **    2019-06-10 https://sqlite.org/src/info/fd76310a5e843e07
+        **    2019-06-14 https://sqlite.org/src/info/ce8717f0885af975
+        **    2019-09-03 https://sqlite.org/src/info/0f0428096f17252a
         */
-        if( sqlite3Isdigit(zNew[0])
-         || zNew[0]=='-'
-         || (zNew[0]+1=='0' && iTo==1)
+        if( pLeft->op!=TK_COLUMN 
+         || sqlite3ExprAffinity(pLeft)!=SQLITE_AFF_TEXT 
+         || IsVirtual(pLeft->y.pTab)  /* Value might be numeric */
         ){
-          if( pLeft->op!=TK_COLUMN 
-           || sqlite3ExprAffinity(pLeft)!=SQLITE_AFF_TEXT 
-           || IsVirtual(pLeft->y.pTab)  /* Value might be numeric */
-          ){
+          int isNum;
+          double rDummy;
+          isNum = sqlite3AtoF(zNew, &rDummy, iTo, SQLITE_UTF8);
+          if( isNum<=0 ){
+            if( iTo==1 && zNew[0]=='-' ){
+              isNum = +1;
+            }else{
+              zNew[iTo-1]++;
+              isNum = sqlite3AtoF(zNew, &rDummy, iTo, SQLITE_UTF8);
+              zNew[iTo-1]--;
+            }
+          }
+          if( isNum>0 ){
             sqlite3ExprDelete(db, pPrefix);
             sqlite3ValueFree(pVal);
             return 0;
@@ -139338,7 +143119,7 @@ static int termIsEquivalence(Parse *pParse, Expr *pExpr){
   ){
     return 0;
   }
-  pColl = sqlite3BinaryCompareCollSeq(pParse, pExpr->pLeft, pExpr->pRight);
+  pColl = sqlite3ExprCompareCollSeq(pParse, pExpr);
   if( sqlite3IsBinary(pColl) ) return 1;
   return sqlite3ExprCollSeqMatch(pParse, pExpr->pLeft, pExpr->pRight);
 }
@@ -139562,7 +143343,7 @@ static void exprAnalyze(
         pDup = pExpr;
         pNew = pTerm;
       }
-      exprCommute(pParse, pDup);
+      pNew->wtFlags |= exprCommute(pParse, pDup);
       pNew->leftCursor = aiCurCol[0];
       pNew->u.leftColumn = aiCurCol[1];
       testcase( (prereqLeft | extraRight) != prereqLeft );
@@ -139731,6 +143512,7 @@ static void exprAnalyze(
             0, sqlite3ExprDup(db, pRight, 0));
         if( ExprHasProperty(pExpr, EP_FromJoin) && pNewExpr ){
           ExprSetProperty(pNewExpr, EP_FromJoin);
+          pNewExpr->iRightJoinTable = pExpr->iRightJoinTable;
         }
         idxNew = whereClauseInsert(pWC, pNewExpr, TERM_VIRTUAL|TERM_DYNAMIC);
         testcase( idxNew==0 );
@@ -139787,11 +143569,15 @@ static void exprAnalyze(
   ** expression). The WhereTerm.iField variable identifies the index within
   ** the vector on the LHS that the virtual term represents.
   **
-  ** This only works if the RHS is a simple SELECT, not a compound
+  ** This only works if the RHS is a simple SELECT (not a compound) that does
+  ** not use window functions.
   */
   if( pWC->op==TK_AND && pExpr->op==TK_IN && pTerm->iField==0
    && pExpr->pLeft->op==TK_VECTOR
    && pExpr->x.pSelect->pPrior==0
+#ifndef SQLITE_OMIT_WINDOWFUNC
+   && pExpr->x.pSelect->pWin==0
+#endif
   ){
     int i;
     for(i=0; i<sqlite3ExprVectorSize(pExpr->pLeft); i++){
@@ -139803,8 +143589,8 @@ static void exprAnalyze(
     }
   }
 
-#ifdef SQLITE_ENABLE_STAT3_OR_STAT4
-  /* When sqlite_stat3 histogram data is available an operator of the
+#ifdef SQLITE_ENABLE_STAT4
+  /* When sqlite_stat4 histogram data is available an operator of the
   ** form "x IS NOT NULL" can sometimes be evaluated more efficiently
   ** as "x>NULL" if x is not an INTEGER PRIMARY KEY.  So construct a
   ** virtual term of that form.
@@ -139815,7 +143601,7 @@ static void exprAnalyze(
    && pExpr->pLeft->op==TK_COLUMN
    && pExpr->pLeft->iColumn>=0
    && !ExprHasProperty(pExpr, EP_FromJoin)
-   && OptimizationEnabled(db, SQLITE_Stat34)
+   && OptimizationEnabled(db, SQLITE_Stat4)
   ){
     Expr *pNewExpr;
     Expr *pLeft = pExpr->pLeft;
@@ -139840,7 +143626,7 @@ static void exprAnalyze(
       pNewTerm->prereqAll = pTerm->prereqAll;
     }
   }
-#endif /* SQLITE_ENABLE_STAT3_OR_STAT4 */
+#endif /* SQLITE_ENABLE_STAT4 */
 
   /* Prevent ON clause terms of a LEFT JOIN from being used to drive
   ** an index for tables to the left of the join.
@@ -139873,7 +143659,7 @@ static void exprAnalyze(
 ** all terms of the WHERE clause.
 */
 SQLITE_PRIVATE void sqlite3WhereSplit(WhereClause *pWC, Expr *pExpr, u8 op){
-  Expr *pE2 = sqlite3ExprSkipCollate(pExpr);
+  Expr *pE2 = sqlite3ExprSkipCollateAndLikely(pExpr);
   pWC->op = op;
   if( pE2==0 ) return;
   if( pE2->op!=op ){
@@ -139949,9 +143735,10 @@ SQLITE_PRIVATE Bitmask sqlite3WhereExprUsageNN(WhereMaskSet *pMaskSet, Expr *p){
     mask |= sqlite3WhereExprListUsage(pMaskSet, p->x.pList);
   }
 #ifndef SQLITE_OMIT_WINDOWFUNC
-  if( p->op==TK_FUNCTION && p->y.pWin ){
+  if( (p->op==TK_FUNCTION || p->op==TK_AGG_FUNCTION) && p->y.pWin ){
     mask |= sqlite3WhereExprListUsage(pMaskSet, p->y.pWin->pPartition);
     mask |= sqlite3WhereExprListUsage(pMaskSet, p->y.pWin->pOrderBy);
+    mask |= sqlite3WhereExprUsage(pMaskSet, p->y.pWin->pFilter);
   }
 #endif
   return mask;
@@ -140027,6 +143814,9 @@ SQLITE_PRIVATE void sqlite3WhereTabFuncArgs(
     pRhs = sqlite3PExpr(pParse, TK_UPLUS, 
         sqlite3ExprDup(pParse->db, pArgs->a[j].pExpr, 0), 0);
     pTerm = sqlite3PExpr(pParse, TK_EQ, pColRef, pRhs);
+    if( pItem->fg.jointype & JT_LEFT ){
+      sqlite3SetJoinExpr(pTerm, pItem->iCursor);
+    }
     whereClauseInsert(pWC, pTerm, TERM_DYNAMIC);
   }
 }
@@ -140155,7 +143945,7 @@ SQLITE_PRIVATE int sqlite3WhereBreakLabel(WhereInfo *pWInfo){
 
 /*
 ** Return ONEPASS_OFF (0) if an UPDATE or DELETE statement is unable to
-** operate directly on the rowis returned by a WHERE clause.  Return
+** operate directly on the rowids returned by a WHERE clause.  Return
 ** ONEPASS_SINGLE (1) if the statement can operation directly because only
 ** a single row is to be changed.  Return ONEPASS_MULTI (2) if the one-pass
 ** optimization can be used on multiple 
@@ -140182,6 +143972,14 @@ SQLITE_PRIVATE int sqlite3WhereOkOnePass(WhereInfo *pWInfo, int *aiCur){
   return pWInfo->eOnePass;
 }
 
+/*
+** Return TRUE if the WHERE loop uses the OP_DeferredSeek opcode to move
+** the data cursor to the row selected by the index cursor.
+*/
+SQLITE_PRIVATE int sqlite3WhereUsesDeferredSeek(WhereInfo *pWInfo){
+  return pWInfo->bDeferredSeek;
+}
+
 /*
 ** Move the content of pSrc into pDest
 */
@@ -140288,7 +144086,8 @@ static WhereTerm *whereScanNext(WhereScan *pScan){
         ){
           if( (pTerm->eOperator & WO_EQUIV)!=0
            && pScan->nEquiv<ArraySize(pScan->aiCur)
-           && (pX = sqlite3ExprSkipCollate(pTerm->pExpr->pRight))->op==TK_COLUMN
+           && (pX = sqlite3ExprSkipCollateAndLikely(pTerm->pExpr->pRight))->op
+               ==TK_COLUMN
           ){
             int j;
             for(j=0; j<pScan->nEquiv; j++){
@@ -140313,8 +144112,7 @@ static WhereTerm *whereScanNext(WhereScan *pScan){
                 continue;
               }
               assert(pX->pLeft);
-              pColl = sqlite3BinaryCompareCollSeq(pParse,
-                                                  pX->pLeft, pX->pRight);
+              pColl = sqlite3ExprCompareCollSeq(pParse, pX);
               if( pColl==0 ) pColl = pParse->db->pDfltColl;
               if( sqlite3StrICmp(pColl->zName, pScan->zCollName) ){
                 continue;
@@ -140484,7 +144282,7 @@ static int findIndexCol(
   const char *zColl = pIdx->azColl[iCol];
 
   for(i=0; i<pList->nExpr; i++){
-    Expr *p = sqlite3ExprSkipCollate(pList->a[i].pExpr);
+    Expr *p = sqlite3ExprSkipCollateAndLikely(pList->a[i].pExpr);
     if( p->op==TK_COLUMN
      && p->iColumn==pIdx->aiColumn[iCol]
      && p->iTable==iBase
@@ -140548,7 +144346,7 @@ static int isDistinctRedundant(
   ** current SELECT is a correlated sub-query.
   */
   for(i=0; i<pDistinct->nExpr; i++){
-    Expr *p = sqlite3ExprSkipCollate(pDistinct->a[i].pExpr);
+    Expr *p = sqlite3ExprSkipCollateAndLikely(pDistinct->a[i].pExpr);
     if( p->op==TK_COLUMN && p->iTable==iBase && p->iColumn<0 ) return 1;
   }
 
@@ -140597,17 +144395,17 @@ static LogEst estLog(LogEst N){
 ** opcodes into OP_Copy when the table is being accessed via co-routine 
 ** instead of via table lookup.
 **
-** If the bIncrRowid parameter is 0, then any OP_Rowid instructions on
-** cursor iTabCur are transformed into OP_Null. Or, if bIncrRowid is non-zero,
-** then each OP_Rowid is transformed into an instruction to increment the
-** value stored in its output register.
+** If the iAutoidxCur is not zero, then any OP_Rowid instructions on
+** cursor iTabCur are transformed into OP_Sequence opcode for the
+** iAutoidxCur cursor, in order to generate unique rowids for the
+** automatic index being generated.
 */
 static void translateColumnToCopy(
   Parse *pParse,      /* Parsing context */
   int iStart,         /* Translate from this opcode to the end */
   int iTabCur,        /* OP_Column/OP_Rowid references to this table */
   int iRegister,      /* The first column is in this register */
-  int bIncrRowid      /* If non-zero, transform OP_rowid to OP_AddImm(1) */
+  int iAutoidxCur     /* If non-zero, cursor of autoindex being generated */
 ){
   Vdbe *v = pParse->pVdbe;
   VdbeOp *pOp = sqlite3VdbeGetOp(v, iStart);
@@ -140621,11 +144419,9 @@ static void translateColumnToCopy(
       pOp->p2 = pOp->p3;
       pOp->p3 = 0;
     }else if( pOp->opcode==OP_Rowid ){
-      if( bIncrRowid ){
-        /* Increment the value stored in the P2 operand of the OP_Rowid. */
-        pOp->opcode = OP_AddImm;
-        pOp->p1 = pOp->p2;
-        pOp->p2 = 1;
+      if( iAutoidxCur ){
+        pOp->opcode = OP_Sequence;
+        pOp->p1 = iAutoidxCur;
       }else{
         pOp->opcode = OP_Null;
         pOp->p1 = 0;
@@ -140642,7 +144438,7 @@ static void translateColumnToCopy(
 ** are no-ops.
 */
 #if !defined(SQLITE_OMIT_VIRTUALTABLE) && defined(WHERETRACE_ENABLED)
-static void TRACE_IDX_INPUTS(sqlite3_index_info *p){
+static void whereTraceIndexInfoInputs(sqlite3_index_info *p){
   int i;
   if( !sqlite3WhereTrace ) return;
   for(i=0; i<p->nConstraint; i++){
@@ -140660,7 +144456,7 @@ static void TRACE_IDX_INPUTS(sqlite3_index_info *p){
        p->aOrderBy[i].desc);
   }
 }
-static void TRACE_IDX_OUTPUTS(sqlite3_index_info *p){
+static void whereTraceIndexInfoOutputs(sqlite3_index_info *p){
   int i;
   if( !sqlite3WhereTrace ) return;
   for(i=0; i<p->nConstraint; i++){
@@ -140676,8 +144472,8 @@ static void TRACE_IDX_OUTPUTS(sqlite3_index_info *p){
   sqlite3DebugPrintf("  estimatedRows=%lld\n", p->estimatedRows);
 }
 #else
-#define TRACE_IDX_INPUTS(A)
-#define TRACE_IDX_OUTPUTS(A)
+#define whereTraceIndexInfoInputs(A)
+#define whereTraceIndexInfoOutputs(A)
 #endif
 
 #ifndef SQLITE_OMIT_AUTOMATIC_INDEX
@@ -140772,7 +144568,7 @@ static void constructAutomaticIndex(
      && (pTerm->wtFlags & TERM_VIRTUAL)==0
      && !ExprHasProperty(pExpr, EP_FromJoin)
      && sqlite3ExprIsTableConstant(pExpr, pSrc->iCursor) ){
-      pPartial = sqlite3ExprAnd(pParse->db, pPartial,
+      pPartial = sqlite3ExprAnd(pParse, pPartial,
                                 sqlite3ExprDup(pParse->db, pExpr, 0));
     }
     if( termCanDriveIndex(pTerm, pSrc, notReady) ){
@@ -140837,7 +144633,8 @@ static void constructAutomaticIndex(
         Expr *pX = pTerm->pExpr;
         idxCols |= cMask;
         pIdx->aiColumn[n] = pTerm->u.leftColumn;
-        pColl = sqlite3BinaryCompareCollSeq(pParse, pX->pLeft, pX->pRight);
+        pColl = sqlite3ExprCompareCollSeq(pParse, pX);
+        assert( pColl!=0 || pParse->nErr>0 ); /* TH3 collate01.800 */
         pIdx->azColl[n] = pColl ? pColl->zName : sqlite3StrBINARY;
         n++;
       }
@@ -140899,14 +144696,15 @@ static void constructAutomaticIndex(
   if( pTabItem->fg.viaCoroutine ){
     sqlite3VdbeChangeP2(v, addrCounter, regBase+n);
     testcase( pParse->db->mallocFailed );
+    assert( pLevel->iIdxCur>0 );
     translateColumnToCopy(pParse, addrTop, pLevel->iTabCur,
-                          pTabItem->regResult, 1);
+                          pTabItem->regResult, pLevel->iIdxCur);
     sqlite3VdbeGoto(v, addrTop);
     pTabItem->fg.viaCoroutine = 0;
   }else{
     sqlite3VdbeAddOp2(v, OP_Next, pLevel->iTabCur, addrTop+1); VdbeCoverage(v);
+    sqlite3VdbeChangeP5(v, SQLITE_STMTSTATUS_AUTOINDEX);
   }
-  sqlite3VdbeChangeP5(v, SQLITE_STMTSTATUS_AUTOINDEX);
   sqlite3VdbeJumpHere(v, addrTop);
   sqlite3ReleaseTempReg(pParse, regRecord);
   
@@ -140969,6 +144767,7 @@ static sqlite3_index_info *allocateIndexInfo(
     for(i=0; i<n; i++){
       Expr *pExpr = pOrderBy->a[i].pExpr;
       if( pExpr->op!=TK_COLUMN || pExpr->iTable!=pSrc->iCursor ) break;
+      if( pOrderBy->a[i].sortFlags & KEYINFO_ORDER_BIGNULL ) break;
     }
     if( i==n){
       nOrderBy = n;
@@ -140984,23 +144783,14 @@ static sqlite3_index_info *allocateIndexInfo(
     sqlite3ErrorMsg(pParse, "out of memory");
     return 0;
   }
-
-  /* Initialize the structure.  The sqlite3_index_info structure contains
-  ** many fields that are declared "const" to prevent xBestIndex from
-  ** changing them.  We have to do some funky casting in order to
-  ** initialize those fields.
-  */
   pHidden = (struct HiddenIndexInfo*)&pIdxInfo[1];
   pIdxCons = (struct sqlite3_index_constraint*)&pHidden[1];
   pIdxOrderBy = (struct sqlite3_index_orderby*)&pIdxCons[nTerm];
   pUsage = (struct sqlite3_index_constraint_usage*)&pIdxOrderBy[nOrderBy];
-  *(int*)&pIdxInfo->nConstraint = nTerm;
-  *(int*)&pIdxInfo->nOrderBy = nOrderBy;
-  *(struct sqlite3_index_constraint**)&pIdxInfo->aConstraint = pIdxCons;
-  *(struct sqlite3_index_orderby**)&pIdxInfo->aOrderBy = pIdxOrderBy;
-  *(struct sqlite3_index_constraint_usage**)&pIdxInfo->aConstraintUsage =
-                                                                   pUsage;
-
+  pIdxInfo->nOrderBy = nOrderBy;
+  pIdxInfo->aConstraint = pIdxCons;
+  pIdxInfo->aOrderBy = pIdxOrderBy;
+  pIdxInfo->aConstraintUsage = pUsage;
   pHidden->pWC = pWC;
   pHidden->pParse = pParse;
   for(i=j=0, pTerm=pWC->a; i<pWC->nTerm; i++, pTerm++){
@@ -141014,18 +144804,13 @@ static sqlite3_index_info *allocateIndexInfo(
     testcase( pTerm->eOperator & WO_ALL );
     if( (pTerm->eOperator & ~(WO_EQUIV))==0 ) continue;
     if( pTerm->wtFlags & TERM_VNULL ) continue;
+
+    /* tag-20191211-002: WHERE-clause constraints are not useful to the
+    ** right-hand table of a LEFT JOIN.  See tag-20191211-001 for the
+    ** equivalent restriction for ordinary tables. */
     if( (pSrc->fg.jointype & JT_LEFT)!=0
      && !ExprHasProperty(pTerm->pExpr, EP_FromJoin)
-     && (pTerm->eOperator & (WO_IS|WO_ISNULL))
     ){
-      /* An "IS" term in the WHERE clause where the virtual table is the rhs
-      ** of a LEFT JOIN. Do not pass this term to the virtual table
-      ** implementation, as this can lead to incorrect results from SQL such
-      ** as:
-      **
-      **   "LEFT JOIN vtab WHERE vtab.col IS NULL"  */
-      testcase( pTerm->eOperator & WO_ISNULL );
-      testcase( pTerm->eOperator & WO_IS );
       continue;
     }
     assert( pTerm->u.leftColumn>=(-1) );
@@ -141056,7 +144841,8 @@ static sqlite3_index_info *allocateIndexInfo(
       if( op & (WO_LT|WO_LE|WO_GT|WO_GE)
        && sqlite3ExprIsVector(pTerm->pExpr->pRight) 
       ){
-        if( i<16 ) mNoOmit |= (1 << i);
+        testcase( j!=i );
+        if( j<16 ) mNoOmit |= (1 << j);
         if( op==WO_LT ) pIdxCons[j].op = WO_LE;
         if( op==WO_GT ) pIdxCons[j].op = WO_GE;
       }
@@ -141064,10 +144850,11 @@ static sqlite3_index_info *allocateIndexInfo(
 
     j++;
   }
+  pIdxInfo->nConstraint = j;
   for(i=0; i<nOrderBy; i++){
     Expr *pExpr = pOrderBy->a[i].pExpr;
     pIdxOrderBy[i].iColumn = pExpr->iColumn;
-    pIdxOrderBy[i].desc = pOrderBy->a[i].sortOrder;
+    pIdxOrderBy[i].desc = pOrderBy->a[i].sortFlags & KEYINFO_ORDER_DESC;
   }
 
   *pmNoOmit = mNoOmit;
@@ -141094,9 +144881,9 @@ static int vtabBestIndex(Parse *pParse, Table *pTab, sqlite3_index_info *p){
   sqlite3_vtab *pVtab = sqlite3GetVTable(pParse->db, pTab)->pVtab;
   int rc;
 
-  TRACE_IDX_INPUTS(p);
+  whereTraceIndexInfoInputs(p);
   rc = pVtab->pModule->xBestIndex(pVtab, p);
-  TRACE_IDX_OUTPUTS(p);
+  whereTraceIndexInfoOutputs(p);
 
   if( rc!=SQLITE_OK && rc!=SQLITE_CONSTRAINT ){
     if( rc==SQLITE_NOMEM ){
@@ -141113,7 +144900,7 @@ static int vtabBestIndex(Parse *pParse, Table *pTab, sqlite3_index_info *p){
 }
 #endif /* !defined(SQLITE_OMIT_VIRTUALTABLE) */
 
-#ifdef SQLITE_ENABLE_STAT3_OR_STAT4
+#ifdef SQLITE_ENABLE_STAT4
 /*
 ** Estimate the location of a particular key among all keys in an
 ** index.  Store the results in aStat as follows:
@@ -141306,7 +145093,7 @@ static int whereKeyStats(
   pRec->nField = nField;
   return i;
 }
-#endif /* SQLITE_ENABLE_STAT3_OR_STAT4 */
+#endif /* SQLITE_ENABLE_STAT4 */
 
 /*
 ** If it is not NULL, pTerm is a term that provides an upper or lower
@@ -141332,7 +145119,7 @@ static LogEst whereRangeAdjust(WhereTerm *pTerm, LogEst nNew){
 }
 
 
-#ifdef SQLITE_ENABLE_STAT3_OR_STAT4
+#ifdef SQLITE_ENABLE_STAT4
 /*
 ** Return the affinity for a single column of an index.
 */
@@ -141341,12 +145128,13 @@ SQLITE_PRIVATE char sqlite3IndexColumnAffinity(sqlite3 *db, Index *pIdx, int iCo
   if( !pIdx->zColAff ){
     if( sqlite3IndexAffinityStr(db, pIdx)==0 ) return SQLITE_AFF_BLOB;
   }
+  assert( pIdx->zColAff[iCol]!=0 );
   return pIdx->zColAff[iCol];
 }
 #endif
 
 
-#ifdef SQLITE_ENABLE_STAT3_OR_STAT4
+#ifdef SQLITE_ENABLE_STAT4
 /* 
 ** This function is called to estimate the number of rows visited by a
 ** range-scan on a skip-scan index. For example:
@@ -141452,7 +145240,7 @@ static int whereRangeSkipScanEst(
 
   return rc;
 }
-#endif /* SQLITE_ENABLE_STAT3_OR_STAT4 */
+#endif /* SQLITE_ENABLE_STAT4 */
 
 /*
 ** This function is used to estimate the number of rows that will be visited
@@ -141505,12 +145293,12 @@ static int whereRangeScanEst(
   int nOut = pLoop->nOut;
   LogEst nNew;
 
-#ifdef SQLITE_ENABLE_STAT3_OR_STAT4
+#ifdef SQLITE_ENABLE_STAT4
   Index *p = pLoop->u.btree.pIndex;
   int nEq = pLoop->u.btree.nEq;
 
-  if( p->nSample>0 && nEq<p->nSampleCol
-   && OptimizationEnabled(pParse->db, SQLITE_Stat34)
+  if( p->nSample>0 && ALWAYS(nEq<p->nSampleCol)
+   && OptimizationEnabled(pParse->db, SQLITE_Stat4)
   ){
     if( nEq==pBuilder->nRecValid ){
       UnpackedRecord *pRec = pBuilder->pRec;
@@ -141608,7 +145396,7 @@ static int whereRangeScanEst(
           /* TUNING:  If both iUpper and iLower are derived from the same
           ** sample, then assume they are 4x more selective.  This brings
           ** the estimated selectivity more in line with what it would be
-          ** if estimated without the use of STAT3/4 tables. */
+          ** if estimated without the use of STAT4 tables. */
           if( iLwrIdx==iUprIdx ) nNew -= 20;  assert( 20==sqlite3LogEst(4) );
         }else{
           nNew = 10;        assert( 10==sqlite3LogEst(2) );
@@ -141657,12 +145445,12 @@ static int whereRangeScanEst(
   return rc;
 }
 
-#ifdef SQLITE_ENABLE_STAT3_OR_STAT4
+#ifdef SQLITE_ENABLE_STAT4
 /*
 ** Estimate the number of rows that will be returned based on
 ** an equality constraint x=VALUE and where that VALUE occurs in
 ** the histogram data.  This only works when x is the left-most
-** column of an index and sqlite_stat3 histogram data is available
+** column of an index and sqlite_stat4 histogram data is available
 ** for that index.  When pExpr==NULL that means the constraint is
 ** "x IS NULL" instead of "x=VALUE".
 **
@@ -141720,9 +145508,9 @@ static int whereEqualScanEst(
   
   return rc;
 }
-#endif /* SQLITE_ENABLE_STAT3_OR_STAT4 */
+#endif /* SQLITE_ENABLE_STAT4 */
 
-#ifdef SQLITE_ENABLE_STAT3_OR_STAT4
+#ifdef SQLITE_ENABLE_STAT4
 /*
 ** Estimate the number of rows that will be returned based on
 ** an IN constraint where the right-hand side of the IN operator
@@ -141769,23 +145557,24 @@ static int whereInScanEst(
   assert( pBuilder->nRecValid==nRecValid );
   return rc;
 }
-#endif /* SQLITE_ENABLE_STAT3_OR_STAT4 */
+#endif /* SQLITE_ENABLE_STAT4 */
 
 
 #ifdef WHERETRACE_ENABLED
 /*
 ** Print the content of a WhereTerm object
 */
-static void whereTermPrint(WhereTerm *pTerm, int iTerm){
+SQLITE_PRIVATE void sqlite3WhereTermPrint(WhereTerm *pTerm, int iTerm){
   if( pTerm==0 ){
     sqlite3DebugPrintf("TERM-%-3d NULL\n", iTerm);
   }else{
-    char zType[4];
+    char zType[8];
     char zLeft[50];
-    memcpy(zType, "...", 4);
+    memcpy(zType, "....", 5);
     if( pTerm->wtFlags & TERM_VIRTUAL ) zType[0] = 'V';
     if( pTerm->eOperator & WO_EQUIV  ) zType[1] = 'E';
     if( ExprHasProperty(pTerm->pExpr, EP_FromJoin) ) zType[2] = 'L';
+    if( pTerm->wtFlags & TERM_CODED  ) zType[3] = 'C';
     if( pTerm->eOperator & WO_SINGLE ){
       sqlite3_snprintf(sizeof(zLeft),zLeft,"left={%d:%d}",
                        pTerm->leftCursor, pTerm->u.leftColumn);
@@ -141796,14 +145585,21 @@ static void whereTermPrint(WhereTerm *pTerm, int iTerm){
       sqlite3_snprintf(sizeof(zLeft),zLeft,"left=%d", pTerm->leftCursor);
     }
     sqlite3DebugPrintf(
-       "TERM-%-3d %p %s %-12s prob=%-3d op=0x%03x wtFlags=0x%04x",
-       iTerm, pTerm, zType, zLeft, pTerm->truthProb,
-       pTerm->eOperator, pTerm->wtFlags);
+       "TERM-%-3d %p %s %-12s op=%03x wtFlags=%04x",
+       iTerm, pTerm, zType, zLeft, pTerm->eOperator, pTerm->wtFlags);
+    /* The 0x10000 .wheretrace flag causes extra information to be
+    ** shown about each Term */
+    if( sqlite3WhereTrace & 0x10000 ){
+      sqlite3DebugPrintf(" prob=%-3d prereq=%llx,%llx",
+        pTerm->truthProb, (u64)pTerm->prereqAll, (u64)pTerm->prereqRight);
+    }
     if( pTerm->iField ){
-      sqlite3DebugPrintf(" iField=%d\n", pTerm->iField);
-    }else{
-      sqlite3DebugPrintf("\n");
+      sqlite3DebugPrintf(" iField=%d", pTerm->iField);
+    }
+    if( pTerm->iParent>=0 ){
+      sqlite3DebugPrintf(" iParent=%d", pTerm->iParent);
     }
+    sqlite3DebugPrintf("\n");
     sqlite3TreeViewExpr(0, pTerm->pExpr, 0);
   }
 }
@@ -141816,7 +145612,7 @@ static void whereTermPrint(WhereTerm *pTerm, int iTerm){
 SQLITE_PRIVATE void sqlite3WhereClausePrint(WhereClause *pWC){
   int i;
   for(i=0; i<pWC->nTerm; i++){
-    whereTermPrint(&pWC->a[i], i);
+    sqlite3WhereTermPrint(&pWC->a[i], i);
   }
 }
 #endif
@@ -141825,7 +145621,7 @@ SQLITE_PRIVATE void sqlite3WhereClausePrint(WhereClause *pWC){
 /*
 ** Print a WhereLoop object for debugging purposes
 */
-static void whereLoopPrint(WhereLoop *p, WhereClause *pWC){
+SQLITE_PRIVATE void sqlite3WhereLoopPrint(WhereLoop *p, WhereClause *pWC){
   WhereInfo *pWInfo = pWC->pWInfo;
   int nb = 1+(pWInfo->pTabList->nSrc+3)/4;
   struct SrcList_item *pItem = pWInfo->pTabList->a + p->iTab;
@@ -141850,7 +145646,7 @@ static void whereLoopPrint(WhereLoop *p, WhereClause *pWC){
   }else{
     char *z;
     if( p->u.vtab.idxStr ){
-      z = sqlite3_mprintf("(%d,\"%s\",%x)",
+      z = sqlite3_mprintf("(%d,\"%s\",%#x)",
                 p->u.vtab.idxNum, p->u.vtab.idxStr, p->u.vtab.omitMask);
     }else{
       z = sqlite3_mprintf("(%d,%x)", p->u.vtab.idxNum, p->u.vtab.omitMask);
@@ -141867,7 +145663,7 @@ static void whereLoopPrint(WhereLoop *p, WhereClause *pWC){
   if( p->nLTerm && (sqlite3WhereTrace & 0x100)!=0 ){
     int i;
     for(i=0; i<p->nLTerm; i++){
-      whereTermPrint(p->aLTerm[i], i);
+      sqlite3WhereTermPrint(p->aLTerm[i], i);
     }
   }
 }
@@ -141971,6 +145767,7 @@ static void whereInfoFree(sqlite3 *db, WhereInfo *pWInfo){
     pWInfo->pLoops = p->pNextLoop;
     whereLoopDelete(db, p);
   }
+  assert( pWInfo->pExprMods==0 );
   sqlite3DbFreeNN(db, pWInfo);
 }
 
@@ -142172,6 +145969,8 @@ static int whereLoopInsert(WhereLoopBuilder *pBuilder, WhereLoop *pTemplate){
   }
   pBuilder->iPlanLimit--;
 
+  whereLoopAdjustCost(pWInfo->pLoops, pTemplate);
+
   /* If pBuilder->pOrSet is defined, then only keep track of the costs
   ** and prereqs.
   */
@@ -142186,7 +145985,7 @@ static int whereLoopInsert(WhereLoopBuilder *pBuilder, WhereLoop *pTemplate){
 #if WHERETRACE_ENABLED /* 0x8 */
       if( sqlite3WhereTrace & 0x8 ){
         sqlite3DebugPrintf(x?"   or-%d:  ":"   or-X:  ", n);
-        whereLoopPrint(pTemplate, pBuilder->pWC);
+        sqlite3WhereLoopPrint(pTemplate, pBuilder->pWC);
       }
 #endif
     }
@@ -142195,7 +145994,6 @@ static int whereLoopInsert(WhereLoopBuilder *pBuilder, WhereLoop *pTemplate){
 
   /* Look for an existing WhereLoop to replace with pTemplate
   */
-  whereLoopAdjustCost(pWInfo->pLoops, pTemplate);
   ppPrev = whereLoopFindLesser(&pWInfo->pLoops, pTemplate);
 
   if( ppPrev==0 ){
@@ -142204,7 +146002,7 @@ static int whereLoopInsert(WhereLoopBuilder *pBuilder, WhereLoop *pTemplate){
 #if WHERETRACE_ENABLED /* 0x8 */
     if( sqlite3WhereTrace & 0x8 ){
       sqlite3DebugPrintf("   skip: ");
-      whereLoopPrint(pTemplate, pBuilder->pWC);
+      sqlite3WhereLoopPrint(pTemplate, pBuilder->pWC);
     }
 #endif
     return SQLITE_OK;  
@@ -142220,12 +146018,12 @@ static int whereLoopInsert(WhereLoopBuilder *pBuilder, WhereLoop *pTemplate){
   if( sqlite3WhereTrace & 0x8 ){
     if( p!=0 ){
       sqlite3DebugPrintf("replace: ");
-      whereLoopPrint(p, pBuilder->pWC);
+      sqlite3WhereLoopPrint(p, pBuilder->pWC);
       sqlite3DebugPrintf("   with: ");
     }else{
       sqlite3DebugPrintf("    add: ");
     }
-    whereLoopPrint(pTemplate, pBuilder->pWC);
+    sqlite3WhereLoopPrint(pTemplate, pBuilder->pWC);
   }
 #endif
   if( p==0 ){
@@ -142249,7 +146047,7 @@ static int whereLoopInsert(WhereLoopBuilder *pBuilder, WhereLoop *pTemplate){
 #if WHERETRACE_ENABLED /* 0x8 */
       if( sqlite3WhereTrace & 0x8 ){
         sqlite3DebugPrintf(" delete: ");
-        whereLoopPrint(pToDel, pBuilder->pWC);
+        sqlite3WhereLoopPrint(pToDel, pBuilder->pWC);
       }
 #endif
       whereLoopDelete(db, pToDel);
@@ -142301,11 +146099,12 @@ static void whereLoopOutputAdjust(
 ){
   WhereTerm *pTerm, *pX;
   Bitmask notAllowed = ~(pLoop->prereq|pLoop->maskSelf);
-  int i, j, k;
+  int i, j;
   LogEst iReduce = 0;    /* pLoop->nOut should not exceed nRow-iReduce */
 
   assert( (pLoop->wsFlags & WHERE_AUTO_INDEX)==0 );
   for(i=pWC->nTerm, pTerm=pWC->a; i>0; i--, pTerm++){
+    assert( pTerm!=0 );
     if( (pTerm->wtFlags & TERM_VIRTUAL)!=0 ) break;
     if( (pTerm->prereqAll & pLoop->maskSelf)==0 ) continue;
     if( (pTerm->prereqAll & notAllowed)!=0 ) continue;
@@ -142326,6 +146125,7 @@ static void whereLoopOutputAdjust(
         pLoop->nOut--;
         if( pTerm->eOperator&(WO_EQ|WO_IS) ){
           Expr *pRight = pTerm->pExpr->pRight;
+          int k = 0;
           testcase( pTerm->pExpr->op==TK_IS );
           if( sqlite3ExprIsInteger(pRight, &k) && k>=(-1) && k<=1 ){
             k = 10;
@@ -142456,8 +146256,9 @@ static int whereLoopAddBtreeIndex(
 
   pNew = pBuilder->pNew;
   if( db->mallocFailed ) return SQLITE_NOMEM_BKPT;
-  WHERETRACE(0x800, ("BEGIN %s.addBtreeIdx(%s), nEq=%d\n",
-                     pProbe->pTable->zName,pProbe->zName, pNew->u.btree.nEq));
+  WHERETRACE(0x800, ("BEGIN %s.addBtreeIdx(%s), nEq=%d, nSkip=%d\n",
+                     pProbe->pTable->zName,pProbe->zName,
+                     pNew->u.btree.nEq, pNew->nSkip));
 
   assert( (pNew->wsFlags & WHERE_VIRTUALTABLE)==0 );
   assert( (pNew->wsFlags & WHERE_TOP_LIMIT)==0 );
@@ -142489,7 +146290,7 @@ static int whereLoopAddBtreeIndex(
     LogEst rCostIdx;
     LogEst nOutUnadjusted;        /* nOut before IN() and WHERE adjustments */
     int nIn = 0;
-#ifdef SQLITE_ENABLE_STAT3_OR_STAT4
+#ifdef SQLITE_ENABLE_STAT4
     int nRecValid = pBuilder->nRecValid;
 #endif
     if( (eOp==WO_ISNULL || (pTerm->wtFlags&TERM_VNULL)!=0)
@@ -142503,9 +146304,9 @@ static int whereLoopAddBtreeIndex(
     ** to mix with a lower range bound from some other source */
     if( pTerm->wtFlags & TERM_LIKEOPT && pTerm->eOperator==WO_LT ) continue;
 
-    /* Do not allow constraints from the WHERE clause to be used by the
-    ** right table of a LEFT JOIN.  Only constraints in the ON clause are
-    ** allowed */
+    /* tag-20191211-001:  Do not allow constraints from the WHERE clause to
+    ** be used by the right table of a LEFT JOIN.  Only constraints in the
+    ** ON clause are allowed.  See tag-20191211-002 for the vtab equivalent. */
     if( (pSrc->fg.jointype & JT_LEFT)!=0
      && !ExprHasProperty(pTerm->pExpr, EP_FromJoin)
     ){
@@ -142550,8 +146351,6 @@ static int whereLoopAddBtreeIndex(
       }else if( ALWAYS(pExpr->x.pList && pExpr->x.pList->nExpr) ){
         /* "x IN (value, value, ...)" */
         nIn = sqlite3LogEst(pExpr->x.pList->nExpr);
-        assert( nIn>0 );  /* RHS always has 2 or more terms...  The parser
-                          ** changes "x IN (?)" into "x=?". */
       }
       if( pProbe->hasStat1 ){
         LogEst M, logK, safetyMargin;
@@ -142647,7 +146446,7 @@ static int whereLoopAddBtreeIndex(
     ** the value of pNew->nOut to account for pTerm (but not nIn/nInMul).  */
     assert( pNew->nOut==saved_nOut );
     if( pNew->wsFlags & WHERE_COLUMN_RANGE ){
-      /* Adjust nOut using stat3/stat4 data. Or, if there is no stat3/stat4
+      /* Adjust nOut using stat4 data. Or, if there is no stat4
       ** data, using some other estimate.  */
       whereRangeScanEst(pParse, pBuilder, pBtm, pTop, pNew);
     }else{
@@ -142661,13 +146460,13 @@ static int whereLoopAddBtreeIndex(
         pNew->nOut += pTerm->truthProb;
         pNew->nOut -= nIn;
       }else{
-#ifdef SQLITE_ENABLE_STAT3_OR_STAT4
+#ifdef SQLITE_ENABLE_STAT4
         tRowcnt nOut = 0;
         if( nInMul==0 
          && pProbe->nSample 
          && pNew->u.btree.nEq<=pProbe->nSampleCol
          && ((eOp & WO_IN)==0 || !ExprHasProperty(pTerm->pExpr, EP_xIsSelect))
-         && OptimizationEnabled(db, SQLITE_Stat34)
+         && OptimizationEnabled(db, SQLITE_Stat4)
         ){
           Expr *pExpr = pTerm->pExpr;
           if( (eOp & (WO_EQ|WO_ISNULL|WO_IS))!=0 ){
@@ -142704,6 +146503,7 @@ static int whereLoopAddBtreeIndex(
     ** it to pNew->rRun, which is currently set to the cost of the index
     ** seek only. Then, if this is a non-covering index, add the cost of
     ** visiting the rows in the main table.  */
+    assert( pSrc->pTab->szTabRow>0 );
     rCostIdx = pNew->nOut + 1 + (15*pProbe->szIdxRow)/pSrc->pTab->szTabRow;
     pNew->rRun = sqlite3LogEstAdd(rLogSize, rCostIdx);
     if( (pNew->wsFlags & (WHERE_IDX_ONLY|WHERE_IPK))==0 ){
@@ -142729,7 +146529,7 @@ static int whereLoopAddBtreeIndex(
       whereLoopAddBtreeIndex(pBuilder, pSrc, pProbe, nInMul+nIn);
     }
     pNew->nOut = saved_nOut;
-#ifdef SQLITE_ENABLE_STAT3_OR_STAT4
+#ifdef SQLITE_ENABLE_STAT4
     pBuilder->nRecValid = nRecValid;
 #endif
   }
@@ -142755,6 +146555,7 @@ static int whereLoopAddBtreeIndex(
   assert( 42==sqlite3LogEst(18) );
   if( saved_nEq==saved_nSkip
    && saved_nEq+1<pProbe->nKeyCol
+   && saved_nEq==pNew->nLTerm
    && pProbe->noSkipScan==0
    && OptimizationEnabled(db, SQLITE_SkipScan)
    && pProbe->aiRowLogEst[saved_nEq+1]>=42  /* TUNING: Minimum for skip-scan */
@@ -142802,7 +146603,7 @@ static int indexMightHelpWithOrderBy(
   if( pIndex->bUnordered ) return 0;
   if( (pOB = pBuilder->pWInfo->pOrderBy)==0 ) return 0;
   for(ii=0; ii<pOB->nExpr; ii++){
-    Expr *pExpr = sqlite3ExprSkipCollate(pOB->a[ii].pExpr);
+    Expr *pExpr = sqlite3ExprSkipCollateAndLikely(pOB->a[ii].pExpr);
     if( pExpr->op==TK_COLUMN && pExpr->iTable==iCursor ){
       if( pExpr->iColumn<0 ) return 1;
       for(jj=0; jj<pIndex->nKeyCol; jj++){
@@ -142823,18 +146624,25 @@ static int indexMightHelpWithOrderBy(
 /* Check to see if a partial index with pPartIndexWhere can be used
 ** in the current query.  Return true if it can be and false if not.
 */
-static int whereUsablePartialIndex(int iTab, WhereClause *pWC, Expr *pWhere){
+static int whereUsablePartialIndex(
+  int iTab,             /* The table for which we want an index */
+  int isLeft,           /* True if iTab is the right table of a LEFT JOIN */
+  WhereClause *pWC,     /* The WHERE clause of the query */
+  Expr *pWhere          /* The WHERE clause from the partial index */
+){
   int i;
   WhereTerm *pTerm;
   Parse *pParse = pWC->pWInfo->pParse;
   while( pWhere->op==TK_AND ){
-    if( !whereUsablePartialIndex(iTab,pWC,pWhere->pLeft) ) return 0;
+    if( !whereUsablePartialIndex(iTab,isLeft,pWC,pWhere->pLeft) ) return 0;
     pWhere = pWhere->pRight;
   }
   if( pParse->db->flags & SQLITE_EnableQPSG ) pParse = 0;
   for(i=0, pTerm=pWC->a; i<pWC->nTerm; i++, pTerm++){
-    Expr *pExpr = pTerm->pExpr;
+    Expr *pExpr;
+    pExpr = pTerm->pExpr;
     if( (!ExprHasProperty(pExpr, EP_FromJoin) || pExpr->iRightJoinTable==iTab)
+     && (isLeft==0 || ExprHasProperty(pExpr, EP_FromJoin))
      && sqlite3ExprImpliesExpr(pParse, pExpr, pWhere, iTab) 
     ){
       return 1;
@@ -142997,8 +146805,11 @@ static int whereLoopAddBtree(
   for(; rc==SQLITE_OK && pProbe; 
       pProbe=(pSrc->pIBIndex ? 0 : pProbe->pNext), iSortIdx++
   ){
+    int isLeft = (pSrc->fg.jointype & JT_OUTER)!=0;
     if( pProbe->pPartIdxWhere!=0
-     && !whereUsablePartialIndex(pSrc->iCursor, pWC, pProbe->pPartIdxWhere) ){
+     && !whereUsablePartialIndex(pSrc->iCursor, isLeft, pWC,
+                                 pProbe->pPartIdxWhere)
+    ){
       testcase( pNew->iTab!=pSrc->iCursor );  /* See ticket [98d973b8f5] */
       continue;  /* Partial index inappropriate for this query */
     }
@@ -143102,7 +146913,7 @@ static int whereLoopAddBtree(
       ** plan */
       pTab->tabFlags |= TF_StatsUsed;
     }
-#ifdef SQLITE_ENABLE_STAT3_OR_STAT4
+#ifdef SQLITE_ENABLE_STAT4
     sqlite3Stat4ProbeFree(pBuilder->pRec);
     pBuilder->nRecValid = 0;
     pBuilder->pRec = 0;
@@ -143225,7 +147036,14 @@ static int whereLoopAddVirtualOne(
       if( iTerm>mxTerm ) mxTerm = iTerm;
       testcase( iTerm==15 );
       testcase( iTerm==16 );
-      if( iTerm<16 && pUsage[i].omit ) pNew->u.vtab.omitMask |= 1<<iTerm;
+      if( pUsage[i].omit ){
+        if( i<16 && ((1<<i)&mNoOmit)==0 ){
+          testcase( i!=iTerm );
+          pNew->u.vtab.omitMask |= 1<<iTerm;
+        }else{
+          testcase( i!=iTerm );
+        }
+      }
       if( (pTerm->eOperator & WO_IN)!=0 ){
         /* A virtual table that is constrained by an IN clause may not
         ** consume the ORDER BY clause because (1) the order of IN terms
@@ -143238,7 +147056,6 @@ static int whereLoopAddVirtualOne(
       }
     }
   }
-  pNew->u.vtab.omitMask &= ~mNoOmit;
 
   pNew->nLTerm = mxTerm+1;
   for(i=0; i<=mxTerm; i++){
@@ -143295,7 +147112,7 @@ SQLITE_API const char *sqlite3_vtab_collation(sqlite3_index_info *pIdxInfo, int
     int iTerm = pIdxInfo->aConstraint[iCons].iTermOffset;
     Expr *pX = pHidden->pWC->a[iTerm].pExpr;
     if( pX->pLeft ){
-      pC = sqlite3BinaryCompareCollSeq(pHidden->pParse, pX->pLeft, pX->pRight);
+      pC = sqlite3ExprCompareCollSeq(pHidden->pParse, pX);
     }
     zRet = (pC ? pC->zName : sqlite3StrBINARY);
   }
@@ -143520,7 +147337,8 @@ static int whereLoopAddOr(
         if( rc==SQLITE_OK ){
           rc = whereLoopAddOr(&sSubBuild, mPrereq, mUnusable);
         }
-        assert( rc==SQLITE_OK || sCur.n==0 );
+        assert( rc==SQLITE_OK || rc==SQLITE_DONE || sCur.n==0 );
+        testcase( rc==SQLITE_DONE );
         if( sCur.n==0 ){
           sSum.n = 0;
           break;
@@ -143728,10 +147546,12 @@ static i8 wherePathSatisfiesOrderBy(
       pLoop = pLast;
     }
     if( pLoop->wsFlags & WHERE_VIRTUALTABLE ){
-      if( pLoop->u.vtab.isOrdered ) obSat = obDone;
+      if( pLoop->u.vtab.isOrdered && (wctrlFlags & WHERE_DISTINCTBY)==0 ){
+        obSat = obDone;
+      }
       break;
-    }else{
-      pLoop->u.btree.nIdxCol = 0;
+    }else if( wctrlFlags & WHERE_DISTINCTBY ){
+      pLoop->u.btree.nDistinctCol = 0;
     }
     iCur = pWInfo->pTabList->a[pLoop->iTab].iCursor;
 
@@ -143742,7 +147562,7 @@ static i8 wherePathSatisfiesOrderBy(
     */
     for(i=0; i<nOrderBy; i++){
       if( MASKBIT(i) & obSat ) continue;
-      pOBExpr = sqlite3ExprSkipCollate(pOrderBy->a[i].pExpr);
+      pOBExpr = sqlite3ExprSkipCollateAndLikely(pOrderBy->a[i].pExpr);
       if( pOBExpr->op!=TK_COLUMN ) continue;
       if( pOBExpr->iTable!=iCur ) continue;
       pTerm = sqlite3WhereFindTerm(&pWInfo->sWC, iCur, pOBExpr->iColumn,
@@ -143779,7 +147599,8 @@ static i8 wherePathSatisfiesOrderBy(
         assert( nColumn==nKeyCol+1 || !HasRowid(pIndex->pTable) );
         assert( pIndex->aiColumn[nColumn-1]==XN_ROWID
                           || !HasRowid(pIndex->pTable));
-        isOrderDistinct = IsUniqueIndex(pIndex);
+        isOrderDistinct = IsUniqueIndex(pIndex)
+                          && (pLoop->wsFlags & WHERE_SKIPSCAN)==0;
       }
 
       /* Loop through all columns of the index and deal with the ones
@@ -143797,15 +147618,21 @@ static i8 wherePathSatisfiesOrderBy(
           u16 eOp = pLoop->aLTerm[j]->eOperator;
 
           /* Skip over == and IS and ISNULL terms.  (Also skip IN terms when
-          ** doing WHERE_ORDERBY_LIMIT processing). 
+          ** doing WHERE_ORDERBY_LIMIT processing).  Except, IS and ISNULL
+          ** terms imply that the index is not UNIQUE NOT NULL in which case
+          ** the loop need to be marked as not order-distinct because it can
+          ** have repeated NULL rows.
           **
           ** If the current term is a column of an ((?,?) IN (SELECT...)) 
           ** expression for which the SELECT returns more than one column,
           ** check that it is the only column used by this loop. Otherwise,
           ** if it is one of two or more, none of the columns can be
-          ** considered to match an ORDER BY term.  */
+          ** considered to match an ORDER BY term.
+          */
           if( (eOp & eqOpMask)!=0 ){
-            if( eOp & WO_ISNULL ){
+            if( eOp & (WO_ISNULL|WO_IS) ){
+              testcase( eOp & WO_ISNULL );
+              testcase( eOp & WO_IS );
               testcase( isOrderDistinct );
               isOrderDistinct = 0;
             }
@@ -143831,7 +147658,7 @@ static i8 wherePathSatisfiesOrderBy(
         */
         if( pIndex ){
           iColumn = pIndex->aiColumn[j];
-          revIdx = pIndex->aSortOrder[j];
+          revIdx = pIndex->aSortOrder[j] & KEYINFO_ORDER_DESC;
           if( iColumn==pIndex->pTable->iPKey ) iColumn = XN_ROWID;
         }else{
           iColumn = XN_ROWID;
@@ -143855,7 +147682,7 @@ static i8 wherePathSatisfiesOrderBy(
         isMatch = 0;
         for(i=0; bOnce && i<nOrderBy; i++){
           if( MASKBIT(i) & obSat ) continue;
-          pOBExpr = sqlite3ExprSkipCollate(pOrderBy->a[i].pExpr);
+          pOBExpr = sqlite3ExprSkipCollateAndLikely(pOrderBy->a[i].pExpr);
           testcase( wctrlFlags & WHERE_GROUPBY );
           testcase( wctrlFlags & WHERE_DISTINCTBY );
           if( (wctrlFlags & (WHERE_GROUPBY|WHERE_DISTINCTBY))==0 ) bOnce = 0;
@@ -143873,7 +147700,9 @@ static i8 wherePathSatisfiesOrderBy(
             pColl = sqlite3ExprNNCollSeq(pWInfo->pParse, pOrderBy->a[i].pExpr);
             if( sqlite3StrICmp(pColl->zName, pIndex->azColl[j])!=0 ) continue;
           }
-          pLoop->u.btree.nIdxCol = j+1;
+          if( wctrlFlags & WHERE_DISTINCTBY ){
+            pLoop->u.btree.nDistinctCol = j+1;
+          }
           isMatch = 1;
           break;
         }
@@ -143881,13 +147710,22 @@ static i8 wherePathSatisfiesOrderBy(
           /* Make sure the sort order is compatible in an ORDER BY clause.
           ** Sort order is irrelevant for a GROUP BY clause. */
           if( revSet ){
-            if( (rev ^ revIdx)!=pOrderBy->a[i].sortOrder ) isMatch = 0;
+            if( (rev ^ revIdx)!=(pOrderBy->a[i].sortFlags&KEYINFO_ORDER_DESC) ){
+              isMatch = 0;
+            }
           }else{
-            rev = revIdx ^ pOrderBy->a[i].sortOrder;
+            rev = revIdx ^ (pOrderBy->a[i].sortFlags & KEYINFO_ORDER_DESC);
             if( rev ) *pRevMask |= MASKBIT(iLoop);
             revSet = 1;
           }
         }
+        if( isMatch && (pOrderBy->a[i].sortFlags & KEYINFO_ORDER_BIGNULL) ){
+          if( j==pLoop->u.btree.nEq ){
+            pLoop->wsFlags |= WHERE_BIGNULL_SORT;
+          }else{
+            isMatch = 0;
+          }
+        }
         if( isMatch ){
           if( iColumn==XN_ROWID ){
             testcase( distinctColumns==0 );
@@ -144801,8 +148639,19 @@ SQLITE_PRIVATE WhereInfo *sqlite3WhereBegin(
       sqlite3DebugPrintf(", limit: %d", iAuxArg);
     }
     sqlite3DebugPrintf(")\n");
+    if( sqlite3WhereTrace & 0x100 ){
+      Select sSelect;
+      memset(&sSelect, 0, sizeof(sSelect));
+      sSelect.selFlags = SF_WhereBegin;
+      sSelect.pSrc = pTabList;
+      sSelect.pWhere = pWhere;
+      sSelect.pOrderBy = pOrderBy;
+      sSelect.pEList = pResultSet;
+      sqlite3TreeViewSelect(0, &sSelect, 0);
+    }
   }
   if( sqlite3WhereTrace & 0x100 ){ /* Display all terms of the WHERE clause */
+    sqlite3DebugPrintf("---- WHERE clause at start of analysis:\n");
     sqlite3WhereClausePrint(sWLB.pWC);
   }
 #endif
@@ -144819,7 +148668,7 @@ SQLITE_PRIVATE WhereInfo *sqlite3WhereBegin(
                                              "ABCDEFGHIJKLMNOPQRSTUVWYXZ";
       for(p=pWInfo->pLoops, i=0; p; p=p->pNextLoop, i++){
         p->cId = zLabel[i%(sizeof(zLabel)-1)];
-        whereLoopPrint(p, sWLB.pWC);
+        sqlite3WhereLoopPrint(p, sWLB.pWC);
       }
     }
 #endif
@@ -144859,7 +148708,7 @@ SQLITE_PRIVATE WhereInfo *sqlite3WhereBegin(
     }
     sqlite3DebugPrintf("\n");
     for(ii=0; ii<pWInfo->nLevel; ii++){
-      whereLoopPrint(pWInfo->a[ii].pWLoop, sWLB.pWC);
+      sqlite3WhereLoopPrint(pWInfo->a[ii].pWLoop, sWLB.pWC);
     }
   }
 #endif
@@ -144884,14 +148733,14 @@ SQLITE_PRIVATE WhereInfo *sqlite3WhereBegin(
   ** then table t2 can be omitted from the following:
   **
   **     SELECT v1, v3 FROM t1 
-  **       LEFT JOIN t2 USING (t1.ipk=t2.ipk)
-  **       LEFT JOIN t3 USING (t1.ipk=t3.ipk)
+  **       LEFT JOIN t2 ON (t1.ipk=t2.ipk)
+  **       LEFT JOIN t3 ON (t1.ipk=t3.ipk)
   **
   ** or from:
   **
   **     SELECT DISTINCT v1, v3 FROM t1 
   **       LEFT JOIN t2
-  **       LEFT JOIN t3 USING (t1.ipk=t3.ipk)
+  **       LEFT JOIN t3 ON (t1.ipk=t3.ipk)
   */
   notReady = ~(Bitmask)0;
   if( pWInfo->nLevel>=2
@@ -144941,7 +148790,13 @@ SQLITE_PRIVATE WhereInfo *sqlite3WhereBegin(
       nTabList--;
     }
   }
+#if defined(WHERETRACE_ENABLED)
+  if( sqlite3WhereTrace & 0x100 ){ /* Display all terms of the WHERE clause */
+    sqlite3DebugPrintf("---- WHERE clause at end of analysis:\n");
+    sqlite3WhereClausePrint(sWLB.pWC);
+  }
   WHERETRACE(0xffff,("*** Optimizer Finished ***\n"));
+#endif
   pWInfo->pParse->nQueryLoop += pWInfo->nRowOut;
 
   /* If the caller is an UPDATE or DELETE statement that is requesting
@@ -145018,7 +148873,13 @@ SQLITE_PRIVATE WhereInfo *sqlite3WhereBegin(
       assert( pTabItem->iCursor==pLevel->iTabCur );
       testcase( pWInfo->eOnePass==ONEPASS_OFF && pTab->nCol==BMS-1 );
       testcase( pWInfo->eOnePass==ONEPASS_OFF && pTab->nCol==BMS );
-      if( pWInfo->eOnePass==ONEPASS_OFF && pTab->nCol<BMS && HasRowid(pTab) ){
+      if( pWInfo->eOnePass==ONEPASS_OFF 
+       && pTab->nCol<BMS
+       && (pTab->tabFlags & (TF_HasGenerated|TF_WithoutRowid))==0
+      ){
+        /* If we know that only a prefix of the record will be used,
+        ** it is advantageous to reduce the "column count" field in
+        ** the P4 operand of the OP_OpenRead/Write opcode. */
         Bitmask b = pTabItem->colUsed;
         int n = 0;
         for(; b; b=b>>1, n++){}
@@ -145077,6 +148938,7 @@ SQLITE_PRIVATE WhereInfo *sqlite3WhereBegin(
         sqlite3VdbeSetP4KeyInfo(pParse, pIx);
         if( (pLoop->wsFlags & WHERE_CONSTRAINT)!=0
          && (pLoop->wsFlags & (WHERE_COLUMN_RANGE|WHERE_SKIPSCAN))==0
+         && (pLoop->wsFlags & WHERE_BIGNULL_SORT)==0
          && (pWInfo->wctrlFlags&WHERE_ORDERBY_MIN)==0
          && pWInfo->eDistinct!=WHERE_DISTINCT_ORDERED
         ){
@@ -145194,7 +149056,7 @@ SQLITE_PRIVATE void sqlite3WhereEnd(WhereInfo *pWInfo){
        && i==pWInfo->nLevel-1  /* Ticket [ef9318757b152e3] 2017-10-21 */
        && (pLoop->wsFlags & WHERE_INDEXED)!=0
        && (pIdx = pLoop->u.btree.pIndex)->hasStat1
-       && (n = pLoop->u.btree.nIdxCol)>0
+       && (n = pLoop->u.btree.nDistinctCol)>0
        && pIdx->aiRowLogEst[n]>=36
       ){
         int r1 = pParse->nMem+1;
@@ -145218,6 +149080,11 @@ SQLITE_PRIVATE void sqlite3WhereEnd(WhereInfo *pWInfo){
       VdbeCoverageIf(v, pLevel->op==OP_Next);
       VdbeCoverageIf(v, pLevel->op==OP_Prev);
       VdbeCoverageIf(v, pLevel->op==OP_VNext);
+      if( pLevel->regBignull ){
+        sqlite3VdbeResolveLabel(v, pLevel->addrBignull);
+        sqlite3VdbeAddOp2(v, OP_DecrJumpZero, pLevel->regBignull, pLevel->p2-1);
+        VdbeCoverage(v);
+      }
 #ifndef SQLITE_DISABLE_SKIPAHEAD_DISTINCT
       if( addrSeek ) sqlite3VdbeJumpHere(v, addrSeek);
 #endif
@@ -145233,10 +149100,26 @@ SQLITE_PRIVATE void sqlite3WhereEnd(WhereInfo *pWInfo){
         if( pIn->eEndLoopOp!=OP_Noop ){
           if( pIn->nPrefix ){
             assert( pLoop->wsFlags & WHERE_IN_EARLYOUT );
-            sqlite3VdbeAddOp4Int(v, OP_IfNoHope, pLevel->iIdxCur,
-                              sqlite3VdbeCurrentAddr(v)+2,
-                              pIn->iBase, pIn->nPrefix);
-            VdbeCoverage(v);
+            if( (pLoop->wsFlags & WHERE_VIRTUALTABLE)==0 ){
+              sqlite3VdbeAddOp4Int(v, OP_IfNoHope, pLevel->iIdxCur,
+                  sqlite3VdbeCurrentAddr(v)+2+(pLevel->iLeftJoin!=0),
+                  pIn->iBase, pIn->nPrefix);
+              VdbeCoverage(v);
+            }
+            if( pLevel->iLeftJoin ){
+              /* For LEFT JOIN queries, cursor pIn->iCur may not have been
+              ** opened yet. This occurs for WHERE clauses such as
+              ** "a = ? AND b IN (...)", where the index is on (a, b). If
+              ** the RHS of the (a=?) is NULL, then the "b IN (...)" may
+              ** never have been coded, but the body of the loop run to
+              ** return the null-row. So, if the cursor is not open yet,
+              ** jump over the OP_Next or OP_Prev instruction about to
+              ** be coded.  */
+              sqlite3VdbeAddOp2(v, OP_IfNotOpen, pIn->iCur, 
+                  sqlite3VdbeCurrentAddr(v) + 2
+              );
+              VdbeCoverage(v);
+            }
           }
           sqlite3VdbeAddOp2(v, pIn->eEndLoopOp, pIn->iCur, pIn->addrInTop);
           VdbeCoverage(v);
@@ -145374,8 +149257,11 @@ SQLITE_PRIVATE void sqlite3WhereEnd(WhereInfo *pWInfo){
             Index *pPk = sqlite3PrimaryKeyIndex(pTab);
             x = pPk->aiColumn[x];
             assert( x>=0 );
+          }else{
+            testcase( x!=sqlite3StorageColumnToTable(pTab,x) );
+            x = sqlite3StorageColumnToTable(pTab,x);
           }
-          x = sqlite3ColumnOfIndex(pIdx, x);
+          x = sqlite3TableColumnToIndex(pIdx, x);
           if( x>=0 ){
             pOp->p2 = x;
             pOp->p1 = pLevel->iIdxCur;
@@ -145398,6 +149284,14 @@ SQLITE_PRIVATE void sqlite3WhereEnd(WhereInfo *pWInfo){
     }
   }
 
+  /* Undo all Expr node modifications */
+  while( pWInfo->pExprMods ){
+    WhereExprMod *p = pWInfo->pExprMods;
+    pWInfo->pExprMods = p->pNext;
+    memcpy(p->pExpr, &p->orig, sizeof(p->orig));
+    sqlite3DbFree(db, p);
+  }
+
   /* Final cleanup
   */
   pParse->nQueryLoop = pWInfo->savedNQueryLoop;
@@ -146145,6 +150039,7 @@ struct WindowRewrite {
   Window *pWin;
   SrcList *pSrc;
   ExprList *pSub;
+  Table *pTab;
   Select *pSubSelect;             /* Current sub-select, if any */
 };
 
@@ -146156,6 +150051,8 @@ struct WindowRewrite {
 static int selectWindowRewriteExprCb(Walker *pWalker, Expr *pExpr){
   struct WindowRewrite *p = pWalker->u.pRewrite;
   Parse *pParse = pWalker->pParse;
+  assert( p!=0 );
+  assert( p->pWin!=0 );
 
   /* If this function is being called from within a scalar sub-select
   ** that used by the SELECT statement being processed, only process
@@ -146193,8 +150090,21 @@ static int selectWindowRewriteExprCb(Walker *pWalker, Expr *pExpr){
 
     case TK_AGG_FUNCTION:
     case TK_COLUMN: {
-      Expr *pDup = sqlite3ExprDup(pParse->db, pExpr, 0);
-      p->pSub = sqlite3ExprListAppend(pParse, p->pSub, pDup);
+      int iCol = -1;
+      if( p->pSub ){
+        int i;
+        for(i=0; i<p->pSub->nExpr; i++){
+          if( 0==sqlite3ExprCompare(0, p->pSub->a[i].pExpr, pExpr, -1) ){
+            iCol = i;
+            break;
+          }
+        }
+      }
+      if( iCol<0 ){
+        Expr *pDup = sqlite3ExprDup(pParse->db, pExpr, 0);
+        if( pDup && pDup->op==TK_AGG_FUNCTION ) pDup->op = TK_FUNCTION;
+        p->pSub = sqlite3ExprListAppend(pParse, p->pSub, pDup);
+      }
       if( p->pSub ){
         assert( ExprHasProperty(pExpr, EP_Static)==0 );
         ExprSetProperty(pExpr, EP_Static);
@@ -146203,10 +150113,11 @@ static int selectWindowRewriteExprCb(Walker *pWalker, Expr *pExpr){
         memset(pExpr, 0, sizeof(Expr));
 
         pExpr->op = TK_COLUMN;
-        pExpr->iColumn = p->pSub->nExpr-1;
+        pExpr->iColumn = (iCol<0 ? p->pSub->nExpr-1: iCol);
         pExpr->iTable = p->pWin->iEphCsr;
+        pExpr->y.pTab = p->pTab;
       }
-
+      if( pParse->db->mallocFailed ) return WRC_Abort;
       break;
     }
 
@@ -146248,17 +150159,20 @@ static void selectWindowRewriteEList(
   Window *pWin,
   SrcList *pSrc,
   ExprList *pEList,               /* Rewrite expressions in this list */
+  Table *pTab,
   ExprList **ppSub                /* IN/OUT: Sub-select expression-list */
 ){
   Walker sWalker;
   WindowRewrite sRewrite;
 
+  assert( pWin!=0 );
   memset(&sWalker, 0, sizeof(Walker));
   memset(&sRewrite, 0, sizeof(WindowRewrite));
 
   sRewrite.pSub = *ppSub;
   sRewrite.pWin = pWin;
   sRewrite.pSrc = pSrc;
+  sRewrite.pTab = pTab;
 
   sWalker.pParse = pParse;
   sWalker.xExprCallback = selectWindowRewriteExprCb;
@@ -146277,15 +150191,23 @@ static void selectWindowRewriteEList(
 static ExprList *exprListAppendList(
   Parse *pParse,          /* Parsing context */
   ExprList *pList,        /* List to which to append. Might be NULL */
-  ExprList *pAppend       /* List of values to append. Might be NULL */
+  ExprList *pAppend,      /* List of values to append. Might be NULL */
+  int bIntToNull
 ){
   if( pAppend ){
     int i;
     int nInit = pList ? pList->nExpr : 0;
     for(i=0; i<pAppend->nExpr; i++){
+      int iDummy;
       Expr *pDup = sqlite3ExprDup(pParse->db, pAppend->a[i].pExpr, 0);
+      assert( pDup==0 || !ExprHasProperty(pDup, EP_MemToken) );
+      if( bIntToNull && pDup && sqlite3ExprIsInteger(pDup, &iDummy) ){
+        pDup->op = TK_NULL;
+        pDup->flags &= ~(EP_IntValue|EP_IsTrue|EP_IsFalse);
+        pDup->u.zToken = 0;
+      }
       pList = sqlite3ExprListAppend(pParse, pList, pDup);
-      if( pList ) pList->a[nInit+i].sortOrder = pAppend->a[i].sortOrder;
+      if( pList ) pList->a[nInit+i].sortFlags = pAppend->a[i].sortFlags;
     }
   }
   return pList;
@@ -146300,7 +150222,7 @@ static ExprList *exprListAppendList(
 */
 SQLITE_PRIVATE int sqlite3WindowRewrite(Parse *pParse, Select *p){
   int rc = SQLITE_OK;
-  if( p->pWin && p->pPrior==0 ){
+  if( p->pWin && p->pPrior==0 && (p->selFlags & SF_WinRewrite)==0 ){
     Vdbe *v = sqlite3GetVdbe(pParse);
     sqlite3 *db = pParse->db;
     Select *pSub = 0;             /* The subquery */
@@ -146313,22 +150235,33 @@ SQLITE_PRIVATE int sqlite3WindowRewrite(Parse *pParse, Select *p){
     ExprList *pSublist = 0;       /* Expression list for sub-query */
     Window *pMWin = p->pWin;      /* Master window object */
     Window *pWin;                 /* Window object iterator */
+    Table *pTab;
+
+    pTab = sqlite3DbMallocZero(db, sizeof(Table));
+    if( pTab==0 ){
+      return sqlite3ErrorToParser(db, SQLITE_NOMEM);
+    }
 
     p->pSrc = 0;
     p->pWhere = 0;
     p->pGroupBy = 0;
     p->pHaving = 0;
+    p->selFlags &= ~SF_Aggregate;
+    p->selFlags |= SF_WinRewrite;
 
     /* Create the ORDER BY clause for the sub-select. This is the concatenation
     ** of the window PARTITION and ORDER BY clauses. Then, if this makes it
     ** redundant, remove the ORDER BY from the parent SELECT.  */
-    pSort = sqlite3ExprListDup(db, pMWin->pPartition, 0);
-    pSort = exprListAppendList(pParse, pSort, pMWin->pOrderBy);
-    if( pSort && p->pOrderBy ){
+    pSort = exprListAppendList(pParse, 0, pMWin->pPartition, 1);
+    pSort = exprListAppendList(pParse, pSort, pMWin->pOrderBy, 1);
+    if( pSort && p->pOrderBy && p->pOrderBy->nExpr<=pSort->nExpr ){
+      int nSave = pSort->nExpr;
+      pSort->nExpr = p->pOrderBy->nExpr;
       if( sqlite3ExprListCompare(pSort, p->pOrderBy, -1)==0 ){
         sqlite3ExprListDelete(db, p->pOrderBy);
         p->pOrderBy = 0;
       }
+      pSort->nExpr = nSave;
     }
 
     /* Assign a cursor number for the ephemeral table used to buffer rows.
@@ -146337,23 +150270,30 @@ SQLITE_PRIVATE int sqlite3WindowRewrite(Parse *pParse, Select *p){
     pMWin->iEphCsr = pParse->nTab++;
     pParse->nTab += 3;
 
-    selectWindowRewriteEList(pParse, pMWin, pSrc, p->pEList, &pSublist);
-    selectWindowRewriteEList(pParse, pMWin, pSrc, p->pOrderBy, &pSublist);
+    selectWindowRewriteEList(pParse, pMWin, pSrc, p->pEList, pTab, &pSublist);
+    selectWindowRewriteEList(pParse, pMWin, pSrc, p->pOrderBy, pTab, &pSublist);
     pMWin->nBufferCol = (pSublist ? pSublist->nExpr : 0);
 
     /* Append the PARTITION BY and ORDER BY expressions to the to the 
     ** sub-select expression list. They are required to figure out where 
     ** boundaries for partitions and sets of peer rows lie.  */
-    pSublist = exprListAppendList(pParse, pSublist, pMWin->pPartition);
-    pSublist = exprListAppendList(pParse, pSublist, pMWin->pOrderBy);
+    pSublist = exprListAppendList(pParse, pSublist, pMWin->pPartition, 0);
+    pSublist = exprListAppendList(pParse, pSublist, pMWin->pOrderBy, 0);
 
     /* Append the arguments passed to each window function to the
     ** sub-select expression list. Also allocate two registers for each
     ** window function - one for the accumulator, another for interim
     ** results.  */
     for(pWin=pMWin; pWin; pWin=pWin->pNextWin){
-      pWin->iArgCol = (pSublist ? pSublist->nExpr : 0);
-      pSublist = exprListAppendList(pParse, pSublist, pWin->pOwner->x.pList);
+      ExprList *pArgs = pWin->pOwner->x.pList;
+      if( pWin->pFunc->funcFlags & SQLITE_FUNC_SUBTYPE ){
+        selectWindowRewriteEList(pParse, pMWin, pSrc, pArgs, pTab, &pSublist);
+        pWin->iArgCol = (pSublist ? pSublist->nExpr : 0);
+        pWin->bExprArgs = 1;
+      }else{
+        pWin->iArgCol = (pSublist ? pSublist->nExpr : 0);
+        pSublist = exprListAppendList(pParse, pSublist, pArgs, 0);
+      }
       if( pWin->pFilter ){
         Expr *pFilter = sqlite3ExprDup(db, pWin->pFilter, 0);
         pSublist = sqlite3ExprListAppend(pParse, pSublist, pFilter);
@@ -146371,7 +150311,7 @@ SQLITE_PRIVATE int sqlite3WindowRewrite(Parse *pParse, Select *p){
     */
     if( pSublist==0 ){
       pSublist = sqlite3ExprListAppend(pParse, 0, 
-          sqlite3ExprAlloc(db, TK_INTEGER, &sqlite3IntTokens[0], 0)
+        sqlite3Expr(db, TK_INTEGER, "0")
       );
     }
 
@@ -146380,34 +150320,57 @@ SQLITE_PRIVATE int sqlite3WindowRewrite(Parse *pParse, Select *p){
     );
     p->pSrc = sqlite3SrcListAppend(pParse, 0, 0, 0);
     if( p->pSrc ){
+      Table *pTab2;
       p->pSrc->a[0].pSelect = pSub;
       sqlite3SrcListAssignCursors(pParse, p->pSrc);
-      if( sqlite3ExpandSubquery(pParse, &p->pSrc->a[0]) ){
+      pSub->selFlags |= SF_Expanded;
+      pTab2 = sqlite3ResultSetOfSelect(pParse, pSub, SQLITE_AFF_NONE);
+      if( pTab2==0 ){
+        /* Might actually be some other kind of error, but in that case
+        ** pParse->nErr will be set, so if SQLITE_NOMEM is set, we will get
+        ** the correct error message regardless. */
         rc = SQLITE_NOMEM;
       }else{
-        pSub->selFlags |= SF_Expanded;
-        p->selFlags &= ~SF_Aggregate;
-        sqlite3SelectPrep(pParse, pSub, 0);
+        memcpy(pTab, pTab2, sizeof(Table));
+        pTab->tabFlags |= TF_Ephemeral;
+        p->pSrc->a[0].pTab = pTab;
+        pTab = pTab2;
       }
-
-      sqlite3VdbeAddOp2(v, OP_OpenEphemeral, pMWin->iEphCsr, pSublist->nExpr);
-      sqlite3VdbeAddOp2(v, OP_OpenDup, pMWin->iEphCsr+1, pMWin->iEphCsr);
-      sqlite3VdbeAddOp2(v, OP_OpenDup, pMWin->iEphCsr+2, pMWin->iEphCsr);
-      sqlite3VdbeAddOp2(v, OP_OpenDup, pMWin->iEphCsr+3, pMWin->iEphCsr);
     }else{
       sqlite3SelectDelete(db, pSub);
     }
     if( db->mallocFailed ) rc = SQLITE_NOMEM;
+    sqlite3DbFree(db, pTab);
   }
 
+  if( rc ){
+    if( pParse->nErr==0 ){
+      assert( pParse->db->mallocFailed );
+      sqlite3ErrorToParser(pParse->db, SQLITE_NOMEM);
+    }
+    sqlite3SelectReset(pParse, p);
+  }
   return rc;
 }
 
+/*
+** Unlink the Window object from the Select to which it is attached,
+** if it is attached.
+*/
+SQLITE_PRIVATE void sqlite3WindowUnlinkFromSelect(Window *p){
+  if( p->ppThis ){
+    *p->ppThis = p->pNextWin;
+    if( p->pNextWin ) p->pNextWin->ppThis = p->ppThis;
+    p->ppThis = 0;
+  }
+}
+
 /*
 ** Free the Window object passed as the second argument.
 */
 SQLITE_PRIVATE void sqlite3WindowDelete(sqlite3 *db, Window *p){
   if( p ){
+    sqlite3WindowUnlinkFromSelect(p);
     sqlite3ExprDelete(db, p->pFilter);
     sqlite3ExprListDelete(db, p->pPartition);
     sqlite3ExprListDelete(db, p->pOrderBy);
@@ -146585,17 +150548,14 @@ SQLITE_PRIVATE void sqlite3WindowChain(Parse *pParse, Window *pWin, Window *pLis
 SQLITE_PRIVATE void sqlite3WindowAttach(Parse *pParse, Expr *p, Window *pWin){
   if( p ){
     assert( p->op==TK_FUNCTION );
-    /* This routine is only called for the parser.  If pWin was not
-    ** allocated due to an OOM, then the parser would fail before ever
-    ** invoking this routine */
-    if( ALWAYS(pWin) ){
-      p->y.pWin = pWin;
-      ExprSetProperty(p, EP_WinFunc);
-      pWin->pOwner = p;
-      if( p->flags & EP_Distinct ){
-        sqlite3ErrorMsg(pParse,
-           "DISTINCT is not supported for window functions");
-      }
+    assert( pWin );
+    p->y.pWin = pWin;
+    ExprSetProperty(p, EP_WinFunc);
+    pWin->pOwner = p;
+    if( (p->flags & EP_Distinct) && pWin->eFrmType!=TK_FILTER ){
+      sqlite3ErrorMsg(pParse,
+          "DISTINCT is not supported for window functions"
+      );
     }
   }else{
     sqlite3WindowDelete(pParse->db, pWin);
@@ -146603,18 +150563,49 @@ SQLITE_PRIVATE void sqlite3WindowAttach(Parse *pParse, Expr *p, Window *pWin){
 }
 
 /*
-** Return 0 if the two window objects are identical, or non-zero otherwise.
-** Identical window objects can be processed in a single scan.
+** Possibly link window pWin into the list at pSel->pWin (window functions
+** to be processed as part of SELECT statement pSel). The window is linked
+** in if either (a) there are no other windows already linked to this
+** SELECT, or (b) the windows already linked use a compatible window frame.
+*/
+SQLITE_PRIVATE void sqlite3WindowLink(Select *pSel, Window *pWin){
+  if( pSel!=0
+   && (0==pSel->pWin || 0==sqlite3WindowCompare(0, pSel->pWin, pWin, 0))
+  ){
+    pWin->pNextWin = pSel->pWin;
+    if( pSel->pWin ){
+      pSel->pWin->ppThis = &pWin->pNextWin;
+    }
+    pSel->pWin = pWin;
+    pWin->ppThis = &pSel->pWin;
+  }
+}
+
+/*
+** Return 0 if the two window objects are identical, 1 if they are
+** different, or 2 if it cannot be determined if the objects are identical
+** or not. Identical window objects can be processed in a single scan.
 */
-SQLITE_PRIVATE int sqlite3WindowCompare(Parse *pParse, Window *p1, Window *p2){
+SQLITE_PRIVATE int sqlite3WindowCompare(Parse *pParse, Window *p1, Window *p2, int bFilter){
+  int res;
+  if( NEVER(p1==0) || NEVER(p2==0) ) return 1;
   if( p1->eFrmType!=p2->eFrmType ) return 1;
   if( p1->eStart!=p2->eStart ) return 1;
   if( p1->eEnd!=p2->eEnd ) return 1;
   if( p1->eExclude!=p2->eExclude ) return 1;
   if( sqlite3ExprCompare(pParse, p1->pStart, p2->pStart, -1) ) return 1;
   if( sqlite3ExprCompare(pParse, p1->pEnd, p2->pEnd, -1) ) return 1;
-  if( sqlite3ExprListCompare(p1->pPartition, p2->pPartition, -1) ) return 1;
-  if( sqlite3ExprListCompare(p1->pOrderBy, p2->pOrderBy, -1) ) return 1;
+  if( (res = sqlite3ExprListCompare(p1->pPartition, p2->pPartition, -1)) ){
+    return res;
+  }
+  if( (res = sqlite3ExprListCompare(p1->pOrderBy, p2->pOrderBy, -1)) ){
+    return res;
+  }
+  if( bFilter ){
+    if( (res = sqlite3ExprCompare(pParse, p1->pFilter, p2->pFilter, -1)) ){
+      return res;
+    }
+  }
   return 0;
 }
 
@@ -146624,10 +150615,17 @@ SQLITE_PRIVATE int sqlite3WindowCompare(Parse *pParse, Window *p1, Window *p2){
 ** to begin iterating through the sub-query results. It is used to allocate
 ** and initialize registers and cursors used by sqlite3WindowCodeStep().
 */
-SQLITE_PRIVATE void sqlite3WindowCodeInit(Parse *pParse, Window *pMWin){
+SQLITE_PRIVATE void sqlite3WindowCodeInit(Parse *pParse, Select *pSelect){
+  int nEphExpr = pSelect->pSrc->a[0].pSelect->pEList->nExpr;
+  Window *pMWin = pSelect->pWin;
   Window *pWin;
   Vdbe *v = sqlite3GetVdbe(pParse);
 
+  sqlite3VdbeAddOp2(v, OP_OpenEphemeral, pMWin->iEphCsr, nEphExpr);
+  sqlite3VdbeAddOp2(v, OP_OpenDup, pMWin->iEphCsr+1, pMWin->iEphCsr);
+  sqlite3VdbeAddOp2(v, OP_OpenDup, pMWin->iEphCsr+2, pMWin->iEphCsr);
+  sqlite3VdbeAddOp2(v, OP_OpenDup, pMWin->iEphCsr+3, pMWin->iEphCsr);
+
   /* Allocate registers to use for PARTITION BY values, if any. Initialize
   ** said registers to NULL.  */
   if( pMWin->pPartition ){
@@ -146666,8 +150664,8 @@ SQLITE_PRIVATE void sqlite3WindowCodeInit(Parse *pParse, Window *pMWin){
       pWin->regApp = pParse->nMem+1;
       pParse->nMem += 3;
       if( pKeyInfo && pWin->pFunc->zName[1]=='i' ){
-        assert( pKeyInfo->aSortOrder[0]==0 );
-        pKeyInfo->aSortOrder[0] = 1;
+        assert( pKeyInfo->aSortFlags[0]==0 );
+        pKeyInfo->aSortFlags[0] = KEYINFO_ORDER_DESC;
       }
       sqlite3VdbeAddOp2(v, OP_OpenEphemeral, pWin->csrApp, 2);
       sqlite3VdbeAppendP4(v, pKeyInfo, P4_KEYINFO);
@@ -146752,6 +150750,108 @@ static int windowArgCount(Window *pWin){
   return (pList ? pList->nExpr : 0);
 }
 
+typedef struct WindowCodeArg WindowCodeArg;
+typedef struct WindowCsrAndReg WindowCsrAndReg;
+
+/*
+** See comments above struct WindowCodeArg.
+*/
+struct WindowCsrAndReg {
+  int csr;                        /* Cursor number */
+  int reg;                        /* First in array of peer values */
+};
+
+/*
+** A single instance of this structure is allocated on the stack by 
+** sqlite3WindowCodeStep() and a pointer to it passed to the various helper
+** routines. This is to reduce the number of arguments required by each
+** helper function.
+**
+** regArg:
+**   Each window function requires an accumulator register (just as an
+**   ordinary aggregate function does). This variable is set to the first
+**   in an array of accumulator registers - one for each window function
+**   in the WindowCodeArg.pMWin list.
+**
+** eDelete:
+**   The window functions implementation sometimes caches the input rows
+**   that it processes in a temporary table. If it is not zero, this
+**   variable indicates when rows may be removed from the temp table (in
+**   order to reduce memory requirements - it would always be safe just
+**   to leave them there). Possible values for eDelete are:
+**
+**      WINDOW_RETURN_ROW:
+**        An input row can be discarded after it is returned to the caller.
+**
+**      WINDOW_AGGINVERSE:
+**        An input row can be discarded after the window functions xInverse()
+**        callbacks have been invoked in it.
+**
+**      WINDOW_AGGSTEP:
+**        An input row can be discarded after the window functions xStep()
+**        callbacks have been invoked in it.
+**
+** start,current,end
+**   Consider a window-frame similar to the following:
+**
+**     (ORDER BY a, b GROUPS BETWEEN 2 PRECEDING AND 2 FOLLOWING)
+**
+**   The windows functions implmentation caches the input rows in a temp
+**   table, sorted by "a, b" (it actually populates the cache lazily, and
+**   aggressively removes rows once they are no longer required, but that's
+**   a mere detail). It keeps three cursors open on the temp table. One
+**   (current) that points to the next row to return to the query engine
+**   once its window function values have been calculated. Another (end)
+**   points to the next row to call the xStep() method of each window function
+**   on (so that it is 2 groups ahead of current). And a third (start) that
+**   points to the next row to call the xInverse() method of each window
+**   function on.
+**
+**   Each cursor (start, current and end) consists of a VDBE cursor
+**   (WindowCsrAndReg.csr) and an array of registers (starting at
+**   WindowCodeArg.reg) that always contains a copy of the peer values 
+**   read from the corresponding cursor.
+**
+**   Depending on the window-frame in question, all three cursors may not
+**   be required. In this case both WindowCodeArg.csr and reg are set to
+**   0.
+*/
+struct WindowCodeArg {
+  Parse *pParse;             /* Parse context */
+  Window *pMWin;             /* First in list of functions being processed */
+  Vdbe *pVdbe;               /* VDBE object */
+  int addrGosub;             /* OP_Gosub to this address to return one row */
+  int regGosub;              /* Register used with OP_Gosub(addrGosub) */
+  int regArg;                /* First in array of accumulator registers */
+  int eDelete;               /* See above */
+
+  WindowCsrAndReg start;
+  WindowCsrAndReg current;
+  WindowCsrAndReg end;
+};
+
+/*
+** Generate VM code to read the window frames peer values from cursor csr into
+** an array of registers starting at reg.
+*/
+static void windowReadPeerValues(
+  WindowCodeArg *p,
+  int csr,
+  int reg
+){
+  Window *pMWin = p->pMWin;
+  ExprList *pOrderBy = pMWin->pOrderBy;
+  if( pOrderBy ){
+    Vdbe *v = sqlite3GetVdbe(p->pParse);
+    ExprList *pPart = pMWin->pPartition;
+    int iColOff = pMWin->nBufferCol + (pPart ? pPart->nExpr : 0);
+    int i;
+    for(i=0; i<pOrderBy->nExpr; i++){
+      sqlite3VdbeAddOp3(v, OP_Column, csr, iColOff+i, reg+i);
+    }
+  }
+}
+
 /*
 ** Generate VM code to invoke either xStep() (if bInverse is 0) or 
 ** xInverse (if bInverse is non-zero) for each window function in the 
@@ -146772,20 +150872,27 @@ static int windowArgCount(Window *pWin){
 ** number of rows in the current partition.
 */
 static void windowAggStep(
-  Parse *pParse, 
+  WindowCodeArg *p,
   Window *pMWin,                  /* Linked list of window functions */
   int csr,                        /* Read arguments from this cursor */
   int bInverse,                   /* True to invoke xInverse instead of xStep */
   int reg                         /* Array of registers */
 ){
+  Parse *pParse = p->pParse;
   Vdbe *v = sqlite3GetVdbe(pParse);
   Window *pWin;
   for(pWin=pMWin; pWin; pWin=pWin->pNextWin){
     FuncDef *pFunc = pWin->pFunc;
     int regArg;
-    int nArg = windowArgCount(pWin);
+    int nArg = pWin->bExprArgs ? 0 : windowArgCount(pWin);
     int i;
 
+    assert( bInverse==0 || pWin->eStart!=TK_UNBOUNDED );
+
+    /* All OVER clauses in the same window function aggregate step must
+    ** be the same. */
+    assert( pWin==pMWin || sqlite3WindowCompare(pParse,pWin,pMWin,0)!=1 );
+
     for(i=0; i<nArg; i++){
       if( i!=1 || pFunc->zName!=nth_valueName ){
         sqlite3VdbeAddOp3(v, OP_Column, csr, pWin->iArgCol+i, reg+i);
@@ -146823,14 +150930,30 @@ static void windowAggStep(
       int addrIf = 0;
       if( pWin->pFilter ){
         int regTmp;
-        assert( nArg==0 || nArg==pWin->pOwner->x.pList->nExpr );
-        assert( nArg || pWin->pOwner->x.pList==0 );
+        assert( pWin->bExprArgs || !nArg ||nArg==pWin->pOwner->x.pList->nExpr );
+        assert( pWin->bExprArgs || nArg  ||pWin->pOwner->x.pList==0 );
         regTmp = sqlite3GetTempReg(pParse);
         sqlite3VdbeAddOp3(v, OP_Column, csr, pWin->iArgCol+nArg,regTmp);
         addrIf = sqlite3VdbeAddOp3(v, OP_IfNot, regTmp, 0, 1);
         VdbeCoverage(v);
         sqlite3ReleaseTempReg(pParse, regTmp);
       }
+      
+      if( pWin->bExprArgs ){
+        int iStart = sqlite3VdbeCurrentAddr(v);
+        VdbeOp *pOp, *pEnd;
+
+        nArg = pWin->pOwner->x.pList->nExpr;
+        regArg = sqlite3GetTempRange(pParse, nArg);
+        sqlite3ExprCodeExprList(pParse, pWin->pOwner->x.pList, regArg, 0, 0);
+
+        pEnd = sqlite3VdbeGetOp(v, -1);
+        for(pOp=sqlite3VdbeGetOp(v, iStart); pOp<=pEnd; pOp++){
+          if( pOp->opcode==OP_Column && pOp->p1==pWin->iEphCsr ){
+            pOp->p1 = csr;
+          }
+        }
+      }
       if( pFunc->funcFlags & SQLITE_FUNC_NEEDCOLL ){
         CollSeq *pColl;
         assert( nArg>0 );
@@ -146841,32 +150964,14 @@ static void windowAggStep(
                         bInverse, regArg, pWin->regAccum);
       sqlite3VdbeAppendP4(v, pFunc, P4_FUNCDEF);
       sqlite3VdbeChangeP5(v, (u8)nArg);
+      if( pWin->bExprArgs ){
+        sqlite3ReleaseTempRange(pParse, regArg, nArg);
+      }
       if( addrIf ) sqlite3VdbeJumpHere(v, addrIf);
     }
   }
 }
 
-typedef struct WindowCodeArg WindowCodeArg;
-typedef struct WindowCsrAndReg WindowCsrAndReg;
-struct WindowCsrAndReg {
-  int csr;
-  int reg;
-};
-
-struct WindowCodeArg {
-  Parse *pParse;
-  Window *pMWin;
-  Vdbe *pVdbe;
-  int regGosub;
-  int addrGosub;
-  int regArg;
-  int eDelete;
-
-  WindowCsrAndReg start;
-  WindowCsrAndReg current;
-  WindowCsrAndReg end;
-};
-
 /*
 ** Values that may be passed as the second argument to windowCodeOp().
 */
@@ -146874,28 +150979,6 @@ struct WindowCodeArg {
 #define WINDOW_AGGINVERSE 2
 #define WINDOW_AGGSTEP    3
 
-/*
-** Generate VM code to read the window frames peer values from cursor csr into
-** an array of registers starting at reg.
-*/
-static void windowReadPeerValues(
-  WindowCodeArg *p,
-  int csr,
-  int reg
-){
-  Window *pMWin = p->pMWin;
-  ExprList *pOrderBy = pMWin->pOrderBy;
-  if( pOrderBy ){
-    Vdbe *v = sqlite3GetVdbe(p->pParse);
-    ExprList *pPart = pMWin->pPartition;
-    int iColOff = pMWin->nBufferCol + (pPart ? pPart->nExpr : 0);
-    int i;
-    for(i=0; i<pOrderBy->nExpr; i++){
-      sqlite3VdbeAddOp3(v, OP_Column, csr, iColOff+i, reg+i);
-    }
-  }
-}
-
 /*
 ** Generate VM code to invoke either xValue() (bFin==0) or xFinalize()
 ** (bFin==1) for each window function in the linked list starting at
@@ -146956,8 +151039,12 @@ static void windowFullScan(WindowCodeArg *p){
   int lblNext;
   int lblBrk;
   int addrNext;
-  int csr = pMWin->csrApp;
+  int csr;
+
+  VdbeModuleComment((v, "windowFullScan begin"));
 
+  assert( pMWin!=0 );
+  csr = pMWin->csrApp;
   nPeer = (pMWin->pOrderBy ? pMWin->pOrderBy->nExpr : 0);
 
   lblNext = sqlite3VdbeMakeLabel(pParse);
@@ -147012,7 +151099,7 @@ static void windowFullScan(WindowCodeArg *p){
     if( addrEq ) sqlite3VdbeJumpHere(v, addrEq);
   }
 
-  windowAggStep(pParse, pMWin, csr, 0, p->regArg);
+  windowAggStep(p, pMWin, csr, 0, p->regArg);
 
   sqlite3VdbeResolveLabel(v, lblNext);
   sqlite3VdbeAddOp2(v, OP_Next, csr, addrNext);
@@ -147027,6 +151114,7 @@ static void windowFullScan(WindowCodeArg *p){
   }
 
   windowAggFinal(p, 1);
+  VdbeModuleComment((v, "windowFullScan end"));
 }
 
 /*
@@ -147201,34 +151289,46 @@ static void windowIfNewPeer(
 /*
 ** This function is called as part of generating VM programs for RANGE
 ** offset PRECEDING/FOLLOWING frame boundaries. Assuming "ASC" order for
-** the ORDER BY term in the window, it generates code equivalent to:
+** the ORDER BY term in the window, and that argument op is OP_Ge, it generates
+** code equivalent to:
 **
 **   if( csr1.peerVal + regVal >= csr2.peerVal ) goto lbl;
 **
-** A special type of arithmetic is used such that if csr.peerVal is not
-** a numeric type (real or integer), then the result of the addition is
-** a copy of csr1.peerVal.
+** The value of parameter op may also be OP_Gt or OP_Le. In these cases the
+** operator in the above pseudo-code is replaced with ">" or "<=", respectively.
+**
+** If the sort-order for the ORDER BY term in the window is DESC, then the
+** comparison is reversed. Instead of adding regVal to csr1.peerVal, it is
+** subtracted. And the comparison operator is inverted to - ">=" becomes "<=",
+** ">" becomes "<", and so on. So, with DESC sort order, if the argument op
+** is OP_Ge, the generated code is equivalent to:
+**
+**   if( csr1.peerVal - regVal <= csr2.peerVal ) goto lbl;
+**
+** A special type of arithmetic is used such that if csr1.peerVal is not
+** a numeric type (real or integer), then the result of the addition addition
+** or subtraction is a a copy of csr1.peerVal.
 */
 static void windowCodeRangeTest(
   WindowCodeArg *p, 
-  int op,                          /* OP_Ge or OP_Gt */
-  int csr1, 
-  int regVal, 
-  int csr2,
-  int lbl
+  int op,                         /* OP_Ge, OP_Gt, or OP_Le */
+  int csr1,                       /* Cursor number for cursor 1 */
+  int regVal,                     /* Register containing non-negative number */
+  int csr2,                       /* Cursor number for cursor 2 */
+  int lbl                         /* Jump destination if condition is true */
 ){
   Parse *pParse = p->pParse;
   Vdbe *v = sqlite3GetVdbe(pParse);
-  int reg1 = sqlite3GetTempReg(pParse);
-  int reg2 = sqlite3GetTempReg(pParse);
-  int arith = OP_Add;
-  int addrGe;
-
-  int regString = ++pParse->nMem;
+  ExprList *pOrderBy = p->pMWin->pOrderBy;  /* ORDER BY clause for window */
+  int reg1 = sqlite3GetTempReg(pParse);     /* Reg. for csr1.peerVal+regVal */
+  int reg2 = sqlite3GetTempReg(pParse);     /* Reg. for csr2.peerVal */
+  int regString = ++pParse->nMem;           /* Reg. for constant value '' */
+  int arith = OP_Add;                       /* OP_Add or OP_Subtract */
+  int addrGe;                               /* Jump destination */
 
   assert( op==OP_Ge || op==OP_Gt || op==OP_Le );
-  assert( p->pMWin->pOrderBy && p->pMWin->pOrderBy->nExpr==1 );
-  if( p->pMWin->pOrderBy->a[0].sortOrder ){
+  assert( pOrderBy && pOrderBy->nExpr==1 );
+  if( pOrderBy->a[0].sortFlags & KEYINFO_ORDER_DESC ){
     switch( op ){
       case OP_Ge: op = OP_Le; break;
       case OP_Gt: op = OP_Lt; break;
@@ -147237,27 +151337,95 @@ static void windowCodeRangeTest(
     arith = OP_Subtract;
   }
 
+  /* Read the peer-value from each cursor into a register */
   windowReadPeerValues(p, csr1, reg1);
   windowReadPeerValues(p, csr2, reg2);
 
-  /* Check if the peer value for csr1 value is a text or blob by comparing
-  ** it to the smallest possible string - ''. If it is, jump over the
-  ** OP_Add or OP_Subtract operation and proceed directly to the comparison. */
+  VdbeModuleComment((v, "CodeRangeTest: if( R%d %s R%d %s R%d ) goto lbl",
+      reg1, (arith==OP_Add ? "+" : "-"), regVal,
+      ((op==OP_Ge) ? ">=" : (op==OP_Le) ? "<=" : (op==OP_Gt) ? ">" : "<"), reg2
+  ));
+
+  /* Register reg1 currently contains csr1.peerVal (the peer-value from csr1).
+  ** This block adds (or subtracts for DESC) the numeric value in regVal
+  ** from it. Or, if reg1 is not numeric (it is a NULL, a text value or a blob),
+  ** then leave reg1 as it is. In pseudo-code, this is implemented as:
+  **
+  **   if( reg1>='' ) goto addrGe;
+  **   reg1 = reg1 +/- regVal
+  **   addrGe:
+  **
+  ** Since all strings and blobs are greater-than-or-equal-to an empty string,
+  ** the add/subtract is skipped for these, as required. If reg1 is a NULL,
+  ** then the arithmetic is performed, but since adding or subtracting from
+  ** NULL is always NULL anyway, this case is handled as required too.  */
   sqlite3VdbeAddOp4(v, OP_String8, 0, regString, 0, "", P4_STATIC);
   addrGe = sqlite3VdbeAddOp3(v, OP_Ge, regString, 0, reg1);
   VdbeCoverage(v);
   sqlite3VdbeAddOp3(v, arith, regVal, reg1, reg1);
   sqlite3VdbeJumpHere(v, addrGe);
+
+  /* If the BIGNULL flag is set for the ORDER BY, then it is required to 
+  ** consider NULL values to be larger than all other values, instead of 
+  ** the usual smaller. The VDBE opcodes OP_Ge and so on do not handle this
+  ** (and adding that capability causes a performance regression), so
+  ** instead if the BIGNULL flag is set then cases where either reg1 or
+  ** reg2 are NULL are handled separately in the following block. The code
+  ** generated is equivalent to:
+  **
+  **   if( reg1 IS NULL ){
+  **     if( op==OP_Ge ) goto lbl;
+  **     if( op==OP_Gt && reg2 IS NOT NULL ) goto lbl;
+  **     if( op==OP_Le && reg2 IS NULL ) goto lbl;
+  **   }else if( reg2 IS NULL ){
+  **     if( op==OP_Le ) goto lbl;
+  **   }
+  **
+  ** Additionally, if either reg1 or reg2 are NULL but the jump to lbl is 
+  ** not taken, control jumps over the comparison operator coded below this
+  ** block.  */
+  if( pOrderBy->a[0].sortFlags & KEYINFO_ORDER_BIGNULL ){
+    /* This block runs if reg1 contains a NULL. */
+    int addr = sqlite3VdbeAddOp1(v, OP_NotNull, reg1); VdbeCoverage(v);
+    switch( op ){
+      case OP_Ge: 
+        sqlite3VdbeAddOp2(v, OP_Goto, 0, lbl); 
+        break;
+      case OP_Gt: 
+        sqlite3VdbeAddOp2(v, OP_NotNull, reg2, lbl); 
+        VdbeCoverage(v); 
+        break;
+      case OP_Le: 
+        sqlite3VdbeAddOp2(v, OP_IsNull, reg2, lbl); 
+        VdbeCoverage(v); 
+        break;
+      default: assert( op==OP_Lt ); /* no-op */ break;
+    }
+    sqlite3VdbeAddOp2(v, OP_Goto, 0, sqlite3VdbeCurrentAddr(v)+3);
+
+    /* This block runs if reg1 is not NULL, but reg2 is. */
+    sqlite3VdbeJumpHere(v, addr);
+    sqlite3VdbeAddOp2(v, OP_IsNull, reg2, lbl); VdbeCoverage(v);
+    if( op==OP_Gt || op==OP_Ge ){
+      sqlite3VdbeChangeP2(v, -1, sqlite3VdbeCurrentAddr(v)+1);
+    }
+  }
+
+  /* Compare registers reg2 and reg1, taking the jump if required. Note that
+  ** control skips over this test if the BIGNULL flag is set and either
+  ** reg1 or reg2 contain a NULL value.  */
   sqlite3VdbeAddOp3(v, op, reg2, lbl, reg1); VdbeCoverage(v);
   sqlite3VdbeChangeP5(v, SQLITE_NULLEQ);
+
   assert( op==OP_Ge || op==OP_Gt || op==OP_Lt || op==OP_Le );
   testcase(op==OP_Ge); VdbeCoverageIf(v, op==OP_Ge);
   testcase(op==OP_Lt); VdbeCoverageIf(v, op==OP_Lt);
   testcase(op==OP_Le); VdbeCoverageIf(v, op==OP_Le);
   testcase(op==OP_Gt); VdbeCoverageIf(v, op==OP_Gt);
-
   sqlite3ReleaseTempReg(pParse, reg1);
   sqlite3ReleaseTempReg(pParse, reg2);
+
+  VdbeModuleComment((v, "CodeRangeTest: end"));
 }
 
 /*
@@ -147277,9 +151445,7 @@ static int windowCodeOp(
   Window *pMWin = p->pMWin;
   int ret = 0;
   Vdbe *v = p->pVdbe;
-  int addrIf = 0; 
   int addrContinue = 0;
-  int addrGoto = 0;
   int bPeer = (pMWin->eFrmType!=TK_ROWS);
 
   int lblDone = sqlite3VdbeMakeLabel(pParse);
@@ -147312,7 +151478,7 @@ static int windowCodeOp(
         );
       }
     }else{
-      addrIf = sqlite3VdbeAddOp3(v, OP_IfPos, regCountdown, 0, 1);
+      sqlite3VdbeAddOp3(v, OP_IfPos, regCountdown, lblDone, 1);
       VdbeCoverage(v);
     }
   }
@@ -147321,6 +151487,25 @@ static int windowCodeOp(
     windowAggFinal(p, 0);
   }
   addrContinue = sqlite3VdbeCurrentAddr(v);
+
+  /* If this is a (RANGE BETWEEN a FOLLOWING AND b FOLLOWING) or
+  ** (RANGE BETWEEN b PRECEDING AND a PRECEDING) frame, ensure the 
+  ** start cursor does not advance past the end cursor within the 
+  ** temporary table. It otherwise might, if (a>b).  */
+  if( pMWin->eStart==pMWin->eEnd && regCountdown
+   && pMWin->eFrmType==TK_RANGE && op==WINDOW_AGGINVERSE
+  ){
+    int regRowid1 = sqlite3GetTempReg(pParse);
+    int regRowid2 = sqlite3GetTempReg(pParse);
+    sqlite3VdbeAddOp2(v, OP_Rowid, p->start.csr, regRowid1);
+    sqlite3VdbeAddOp2(v, OP_Rowid, p->end.csr, regRowid2);
+    sqlite3VdbeAddOp3(v, OP_Ge, regRowid2, lblDone, regRowid1);
+    VdbeCoverage(v);
+    sqlite3ReleaseTempReg(pParse, regRowid1);
+    sqlite3ReleaseTempReg(pParse, regRowid2);
+    assert( pMWin->eStart==TK_PRECEDING || pMWin->eStart==TK_FOLLOWING );
+  }
+
   switch( op ){
     case WINDOW_RETURN_ROW:
       csr = p->current.csr;
@@ -147335,7 +151520,7 @@ static int windowCodeOp(
         assert( pMWin->regEndRowid );
         sqlite3VdbeAddOp2(v, OP_AddImm, pMWin->regStartRowid, 1);
       }else{
-        windowAggStep(pParse, pMWin, csr, 1, p->regArg);
+        windowAggStep(p, pMWin, csr, 1, p->regArg);
       }
       break;
 
@@ -147347,7 +151532,7 @@ static int windowCodeOp(
         assert( pMWin->regEndRowid );
         sqlite3VdbeAddOp2(v, OP_AddImm, pMWin->regEndRowid, 1);
       }else{
-        windowAggStep(pParse, pMWin, csr, 0, p->regArg);
+        windowAggStep(p, pMWin, csr, 0, p->regArg);
       }
       break;
   }
@@ -147365,7 +151550,7 @@ static int windowCodeOp(
     sqlite3VdbeAddOp2(v, OP_Next, csr, sqlite3VdbeCurrentAddr(v)+1+bPeer);
     VdbeCoverage(v);
     if( bPeer ){
-      addrGoto = sqlite3VdbeAddOp0(v, OP_Goto);
+      sqlite3VdbeAddOp2(v, OP_Goto, 0, lblDone);
     }
   }
 
@@ -147381,8 +151566,6 @@ static int windowCodeOp(
     sqlite3VdbeAddOp2(v, OP_Goto, 0, addrNextRange);
   }
   sqlite3VdbeResolveLabel(v, lblDone);
-  if( addrGoto ) sqlite3VdbeJumpHere(v, addrGoto);
-  if( addrIf ) sqlite3VdbeJumpHere(v, addrIf);
   return ret;
 }
 
@@ -147398,6 +151581,7 @@ SQLITE_PRIVATE Window *sqlite3WindowDup(sqlite3 *db, Expr *pOwner, Window *p){
     pNew = sqlite3DbMallocZero(db, sizeof(Window));
     if( pNew ){
       pNew->zName = sqlite3DbStrDup(db, p->zName);
+      pNew->zBase = sqlite3DbStrDup(db, p->zBase);
       pNew->pFilter = sqlite3ExprDup(db, p->pFilter, 0);
       pNew->pFunc = p->pFunc;
       pNew->pPartition = sqlite3ExprListDup(db, p->pPartition, 0);
@@ -147406,9 +151590,11 @@ SQLITE_PRIVATE Window *sqlite3WindowDup(sqlite3 *db, Expr *pOwner, Window *p){
       pNew->eEnd = p->eEnd;
       pNew->eStart = p->eStart;
       pNew->eExclude = p->eExclude;
+      pNew->regResult = p->regResult;
       pNew->pStart = sqlite3ExprDup(db, p->pStart, 0);
       pNew->pEnd = sqlite3ExprDup(db, p->pEnd, 0);
       pNew->pOwner = pOwner;
+      pNew->bImplicitFrame = p->bImplicitFrame;
     }
   }
   return pNew;
@@ -147732,7 +151918,7 @@ static int windowExprGtZero(Parse *pParse, Expr *pExpr){
 **         regEnd = <expr2>
 **         regStart = <expr1>
 **       }else{
-**         if( (csrEnd.key + regEnd) <= csrCurrent.key ){
+**         while( (csrEnd.key + regEnd) <= csrCurrent.key ){
 **           AGGSTEP
 **         }
 **         while( (csrStart.key + regStart) < csrCurrent.key ){
@@ -147805,8 +151991,6 @@ SQLITE_PRIVATE void sqlite3WindowCodeStep(
   int addrGosubFlush = 0;         /* Address of OP_Gosub to flush: */
   int addrInteger = 0;            /* Address of OP_Integer */
   int addrEmpty;                  /* Address of OP_Rewind in flush: */
-  int regStart = 0;               /* Value of <expr> PRECEDING */
-  int regEnd = 0;                 /* Value of <expr> FOLLOWING */
   int regNew;                     /* Array of registers holding new input row */
   int regRecord;                  /* regNew array in record form */
   int regRowid;                   /* Rowid for regRecord in eph table */
@@ -147815,6 +151999,8 @@ SQLITE_PRIVATE void sqlite3WindowCodeStep(
   int regFlushPart = 0;           /* Register for "Gosub flush_partition" */
   WindowCodeArg s;                /* Context object for sub-routines */
   int lblWhereEnd;                /* Label just before sqlite3WhereEnd() code */
+  int regStart = 0;               /* Value of <expr> PRECEDING */
+  int regEnd = 0;                 /* Value of <expr> FOLLOWING */
 
   assert( pMWin->eStart==TK_PRECEDING || pMWin->eStart==TK_CURRENT 
        || pMWin->eStart==TK_FOLLOWING || pMWin->eStart==TK_UNBOUNDED 
@@ -147945,14 +152131,14 @@ SQLITE_PRIVATE void sqlite3WindowCodeStep(
 
   if( regStart ){
     sqlite3ExprCode(pParse, pMWin->pStart, regStart);
-    windowCheckValue(pParse, regStart, 0 + (pMWin->eFrmType==TK_RANGE ? 3 : 0));
+    windowCheckValue(pParse, regStart, 0 + (pMWin->eFrmType==TK_RANGE?3:0));
   }
   if( regEnd ){
     sqlite3ExprCode(pParse, pMWin->pEnd, regEnd);
-    windowCheckValue(pParse, regEnd, 1 + (pMWin->eFrmType==TK_RANGE ? 3 : 0));
+    windowCheckValue(pParse, regEnd, 1 + (pMWin->eFrmType==TK_RANGE?3:0));
   }
 
-  if( pMWin->eStart==pMWin->eEnd && regStart ){
+  if( pMWin->eFrmType!=TK_RANGE && pMWin->eStart==pMWin->eEnd && regStart ){
     int op = ((pMWin->eStart==TK_FOLLOWING) ? OP_Ge : OP_Le);
     int addrGe = sqlite3VdbeAddOp3(v, op, regStart, 0, regEnd);
     VdbeCoverageNeverNullIf(v, op==OP_Ge); /* NeverNull because bound <expr> */
@@ -148201,8 +152387,9 @@ struct FrameBound     { int eType; Expr *pExpr; };
 ** shared across database connections.
 */
 static void disableLookaside(Parse *pParse){
+  sqlite3 *db = pParse->db;
   pParse->disableLookaside++;
-  pParse->db->lookaside.bDisable++;
+  DisableLookaside;
 }
 
 
@@ -148212,6 +152399,7 @@ static void disableLookaside(Parse *pParse){
   ** SQLITE_LIMIT_COMPOUND_SELECT.
   */
   static void parserDoubleLinkSelect(Parse *pParse, Select *p){
+    assert( p!=0 );
     if( p->pPrior ){
       Select *pNext = 0, *pLoop;
       int mxSelect, cnt = 0;
@@ -148238,7 +152426,7 @@ static void disableLookaside(Parse *pParse){
     if( p ){
       /* memset(p, 0, sizeof(Expr)); */
       p->op = (u8)op;
-      p->affinity = 0;
+      p->affExpr = 0;
       p->flags = EP_Leaf;
       p->iAgg = -1;
       p->pLeft = p->pRight = 0;
@@ -148365,28 +152553,28 @@ static void disableLookaside(Parse *pParse){
 #endif
 /************* Begin control #defines *****************************************/
 #define YYCODETYPE unsigned short int
-#define YYNOCODE 301
+#define YYNOCODE 310
 #define YYACTIONTYPE unsigned short int
-#define YYWILDCARD 95
+#define YYWILDCARD 100
 #define sqlite3ParserTOKENTYPE Token
 typedef union {
   int yyinit;
   sqlite3ParserTOKENTYPE yy0;
-  With* yy59;
-  IdList* yy62;
-  struct TrigEvent yy90;
-  Upsert* yy136;
-  struct FrameBound yy201;
-  u8 yy238;
-  const char* yy294;
-  Window* yy295;
-  struct {int value; int mask;} yy355;
-  ExprList* yy434;
-  TriggerStep* yy455;
-  Select* yy457;
-  SrcList* yy483;
-  int yy494;
-  Expr* yy524;
+  SrcList* yy47;
+  u8 yy58;
+  struct FrameBound yy77;
+  With* yy131;
+  int yy192;
+  Expr* yy202;
+  struct {int value; int mask;} yy207;
+  struct TrigEvent yy230;
+  ExprList* yy242;
+  Window* yy303;
+  Upsert* yy318;
+  const char* yy436;
+  TriggerStep* yy447;
+  Select* yy539;
+  IdList* yy600;
 } YYMINORTYPE;
 #ifndef YYSTACKDEPTH
 #define YYSTACKDEPTH 100
@@ -148402,17 +152590,18 @@ typedef union {
 #define sqlite3ParserCTX_FETCH Parse *pParse=yypParser->pParse;
 #define sqlite3ParserCTX_STORE yypParser->pParse=pParse;
 #define YYFALLBACK 1
-#define YYNSTATE             541
-#define YYNRULE              375
-#define YYNTOKEN             176
-#define YY_MAX_SHIFT         540
-#define YY_MIN_SHIFTREDUCE   784
-#define YY_MAX_SHIFTREDUCE   1158
-#define YY_ERROR_ACTION      1159
-#define YY_ACCEPT_ACTION     1160
-#define YY_NO_ACTION         1161
-#define YY_MIN_REDUCE        1162
-#define YY_MAX_REDUCE        1536
+#define YYNSTATE             551
+#define YYNRULE              385
+#define YYNRULE_WITH_ACTION  325
+#define YYNTOKEN             181
+#define YY_MAX_SHIFT         550
+#define YY_MIN_SHIFTREDUCE   801
+#define YY_MAX_SHIFTREDUCE   1185
+#define YY_ERROR_ACTION      1186
+#define YY_ACCEPT_ACTION     1187
+#define YY_NO_ACTION         1188
+#define YY_MIN_REDUCE        1189
+#define YY_MAX_REDUCE        1573
 /************* End control #defines *******************************************/
 #define YY_NLOOKAHEAD ((int)(sizeof(yy_lookahead)/sizeof(yy_lookahead[0])))
 
@@ -148479,603 +152668,583 @@ typedef union {
 **  yy_default[]       Default action for each state.
 **
 *********** Begin parsing tables **********************************************/
-#define YY_ACTTAB_COUNT (2142)
+#define YY_ACTTAB_COUNT (1958)
 static const YYACTIONTYPE yy_action[] = {
- /*     0 */   535, 1323,  112,  109,  209,  112,  109,  209, 1160,    1,
- /*    10 */     1,  540,    2, 1164,  535, 1292, 1228, 1207,  289,  384,
- /*    20 */   134,   42,   42, 1427,  382, 1228,    9, 1241,  242,  492,
- /*    30 */  1291,  915,  373,  379, 1026,   70,   70,  427, 1026,  916,
- /*    40 */   529,  529,  529,  119,  120,  110, 1136, 1136,  981,  984,
- /*    50 */   974,  974,  117,  117,  118,  118,  118,  118,  380,  264,
- /*    60 */   264,  264,  264, 1134,  264,  264,  112,  109,  209,  397,
- /*    70 */   454,  517,  532,  491,  532, 1233, 1233,  532,  239,  206,
- /*    80 */   493,  112,  109,  209,  464,  219,  118,  118,  118,  118,
- /*    90 */   111,  393,  440,  444,   16,   16,  116,  116,  116,  116,
- /*   100 */   115,  115,  114,  114,  114,  113,  415,  971,  971,  982,
- /*   110 */   985,  235, 1463,  351, 1134,  419,  384,  116,  116,  116,
- /*   120 */   116,  115,  115,  114,  114,  114,  113,  415,  116,  116,
- /*   130 */   116,  116,  115,  115,  114,  114,  114,  113,  415,  961,
- /*   140 */   119,  120,  110, 1136, 1136,  981,  984,  974,  974,  117,
- /*   150 */   117,  118,  118,  118,  118,  952,  415,  941,  298,  951,
- /*   160 */   941, 1480,  540,    2, 1164, 1115,  535, 1458,  160,  289,
- /*   170 */     6,  134, 1504,  389,  406,  975,  338, 1024, 1241,  337,
- /*   180 */  1089, 1476, 1089,  118,  118,  118,  118,   42,   42,  329,
- /*   190 */   951,  951,  953,  116,  116,  116,  116,  115,  115,  114,
- /*   200 */   114,  114,  113,  415,  311,  430,  299,  311,  881,  160,
- /*   210 */   264,  264,  401,  384,  324, 1115, 1116, 1117,  288,  526,
- /*   220 */    96,  159, 1441,  532,  141,  116,  116,  116,  116,  115,
- /*   230 */   115,  114,  114,  114,  113,  415,  219,  119,  120,  110,
- /*   240 */  1136, 1136,  981,  984,  974,  974,  117,  117,  118,  118,
- /*   250 */   118,  118,  115,  115,  114,  114,  114,  113,  415,  288,
- /*   260 */   526,  403,  533,  121,  870,  870,  419,  250,  267,  336,
- /*   270 */   475,  331,  474,  236,  160,  319, 1084,  322, 1465,  329,
- /*   280 */   350,   12,  535,  384,  502, 1115, 1084,  435,  312, 1084,
- /*   290 */   116,  116,  116,  116,  115,  115,  114,  114,  114,  113,
- /*   300 */   415,  535,  836,   42,   42,  138,  426,  119,  120,  110,
- /*   310 */  1136, 1136,  981,  984,  974,  974,  117,  117,  118,  118,
- /*   320 */   118,  118,   70,   70,  288,  526,  412,  411,  480, 1457,
- /*   330 */   335,   79,    6,  473, 1140, 1115, 1116, 1117,  501, 1142,
- /*   340 */   334,  837,  811, 1484,  512, 1164,  534, 1141,  123,  187,
- /*   350 */   289,  384,  134,  448,  434, 1115,   80,  349,  498, 1241,
- /*   360 */   116,  116,  116,  116,  115,  115,  114,  114,  114,  113,
- /*   370 */   415, 1143, 1115, 1143,  459,  119,  120,  110, 1136, 1136,
- /*   380 */   981,  984,  974,  974,  117,  117,  118,  118,  118,  118,
- /*   390 */   404,  264,  264,  811, 1463,  506,  368, 1156,  535,  114,
- /*   400 */   114,  114,  113,  415,  532, 1115, 1116, 1117,  231,  518,
- /*   410 */  1500,  472,  469,  468,  175,  497,  422,  219, 1202,   70,
- /*   420 */    70,  467, 1115, 1116, 1117,  176,  201,  200,  116,  116,
- /*   430 */   116,  116,  115,  115,  114,  114,  114,  113,  415,  535,
- /*   440 */  1115,  264,  264,  435,  312, 1115,  273,  419,  384,  513,
- /*   450 */  1450, 1115,  326, 1084,  532,  517,   82, 1084,  167,  388,
- /*   460 */    69,   69, 1115, 1084,  519,  509, 1084, 1084,   12, 1157,
- /*   470 */  1084,  420,  119,  120,  110, 1136, 1136,  981,  984,  974,
- /*   480 */   974,  117,  117,  118,  118,  118,  118,  258,  258,  535,
- /*   490 */  1115, 1116, 1117, 1045,  535, 1115, 1116, 1117, 1323,  535,
- /*   500 */   532, 1115, 1116, 1117,  296,  483, 1211,  818, 1046,  448,
- /*   510 */    70,   70, 1115, 1116, 1117,   50,   50,  448,  356,  500,
- /*   520 */    70,   70,  207, 1047,   32,  116,  116,  116,  116,  115,
- /*   530 */   115,  114,  114,  114,  113,  415,  453,  264,  264, 1115,
- /*   540 */   450,  449,  961,  508,  856,  384,  517,    5,  900,  822,
- /*   550 */   532,  484,  181, 1115,  857,  516,  517,  818,  952,  507,
- /*   560 */     3, 1115,  951, 1231, 1231,  482,  398, 1115, 1095,  119,
- /*   570 */   120,  110, 1136, 1136,  981,  984,  974,  974,  117,  117,
- /*   580 */   118,  118,  118,  118, 1115,  535,  238, 1115, 1391, 1115,
- /*   590 */  1116, 1117,  159,  951,  951,  953,  231, 1115,  259,  472,
- /*   600 */   469,  468,  310, 1115, 1116, 1117,   13,   13,  297,  467,
- /*   610 */   276, 1115, 1116, 1117,  412,  411, 1095, 1115, 1116, 1117,
- /*   620 */   395,  355,  116,  116,  116,  116,  115,  115,  114,  114,
- /*   630 */   114,  113,  415,  208, 1115, 1116, 1117, 1115, 1116, 1117,
- /*   640 */   264,  264,  384,  337,  902,  393,  815, 1115, 1116, 1117,
- /*   650 */   413,  413,  413,  532,  112,  109,  209,  309,  900, 1143,
- /*   660 */   535, 1143,  535,  393,  901, 1210,  119,  120,  110, 1136,
- /*   670 */  1136,  981,  984,  974,  974,  117,  117,  118,  118,  118,
- /*   680 */   118,   13,   13,   13,   13,  265,  265,  535,  143,  264,
- /*   690 */   264,  288,  526,  535, 1119,  400,  535,  402,  532,  510,
- /*   700 */  1457,  512,  532,    6,  113,  415, 1067, 1530,   70,   70,
- /*   710 */  1530,  535,  271,  535,   70,   70,  535,   13,   13,  116,
- /*   720 */   116,  116,  116,  115,  115,  114,  114,  114,  113,  415,
- /*   730 */   272,  277,   13,   13,   13,   13,  535,   13,   13,  384,
- /*   740 */   535,  304,  425, 1100,  284, 1119,  184,  801,  185,  338,
- /*   750 */   285,  514, 1532,  369, 1239, 1438, 1182,   70,   70,  425,
- /*   760 */   424,   70,   70,  119,  120,  110, 1136, 1136,  981,  984,
- /*   770 */   974,  974,  117,  117,  118,  118,  118,  118,  190, 1065,
- /*   780 */  1067, 1531,  442,  107, 1531,  408,  264,  264,  264,  264,
- /*   790 */   383, 1396,  261,  410,   95,  900,  485,  414,  421,  532,
- /*   800 */  1045,  532,  301, 1133,  303,  488,  433, 1451, 1396, 1398,
- /*   810 */   278,  535,  278,  520, 1435, 1046,  116,  116,  116,  116,
- /*   820 */   115,  115,  114,  114,  114,  113,  415,  425,  264,  264,
- /*   830 */  1047,  190,   54,   54,  535,  291,  384,  264,  264,  362,
- /*   840 */   962,  532, 1004,  376, 1084,  264,  264, 1029, 1029,  456,
- /*   850 */   532,  523,  270, 1065, 1084,   55,   55, 1084,  532,  442,
- /*   860 */   119,  120,  110, 1136, 1136,  981,  984,  974,  974,  117,
- /*   870 */   117,  118,  118,  118,  118,  535, 1396,  190,  302, 1383,
- /*   880 */   208,  535,  789,  790,  791,  535,  515,  535, 1323,  371,
- /*   890 */   337,  234,  233,  232,  459,  515,   15,   15,  459,  477,
- /*   900 */   459,  459,   44,   44,  136,  900,   56,   56,   57,   57,
- /*   910 */  1185,  390,  197,  116,  116,  116,  116,  115,  115,  114,
- /*   920 */   114,  114,  113,  415,  535,  876,  535,  442,  535,  274,
- /*   930 */   875, 1323,  357,  384,  353,  140, 1426,  946, 1455, 1323,
- /*   940 */  1390,    6, 1240, 1236,  292,   58,   58,   59,   59,   60,
- /*   950 */    60,  535, 1456,  384,  535,    6,  399,  119,  120,  110,
- /*   960 */  1136, 1136,  981,  984,  974,  974,  117,  117,  118,  118,
- /*   970 */   118,  118,   61,   61,  535,   45,   45,  119,  120,  110,
- /*   980 */  1136, 1136,  981,  984,  974,  974,  117,  117,  118,  118,
- /*   990 */   118,  118, 1477,  479,  202,   46,   46,  275,   95,  455,
- /*  1000 */   535,  212,  535,  337,  535, 1454,  535,  409,    6,  242,
- /*  1010 */   116,  116,  116,  116,  115,  115,  114,  114,  114,  113,
- /*  1020 */   415,   48,   48,   49,   49,   62,   62,   63,   63,  535,
- /*  1030 */   116,  116,  116,  116,  115,  115,  114,  114,  114,  113,
- /*  1040 */   415,  535,  459,  535, 1134,  535, 1151,  535,  142,  535,
- /*  1050 */    64,   64,  535, 1338,  535,  494,  535,  446,  535, 1264,
- /*  1060 */   535, 1337,   14,   14,   65,   65,  125,  125,   66,   66,
- /*  1070 */    51,   51,  535,   67,   67,   68,   68,   52,   52,  147,
- /*  1080 */   147,  148,  148, 1453,  317,   98,    6,  535, 1245,  481,
- /*  1090 */   535,  827,  535,   75,   75, 1134,  102,  481,  100,  535,
- /*  1100 */   532,  535,  368, 1066, 1503,  384,  535,  845,   53,   53,
- /*  1110 */    93,   71,   71,  126,  126,  295,  528,  390,  288,  526,
- /*  1120 */    72,   72,  127,  127,  139,  384,   38,  128,  128,  119,
- /*  1130 */   120,  110, 1136, 1136,  981,  984,  974,  974,  117,  117,
- /*  1140 */   118,  118,  118,  118,  535,  495,  535,  447,  535,  119,
- /*  1150 */   120,  110, 1136, 1136,  981,  984,  974,  974,  117,  117,
- /*  1160 */   118,  118,  118,  118,  235,  124,  124,  146,  146,  145,
- /*  1170 */   145,  287,  535, 1277,  535, 1157,  535,  391,  161,  263,
- /*  1180 */   206,  381,  116,  116,  116,  116,  115,  115,  114,  114,
- /*  1190 */   114,  113,  415,  132,  132,  131,  131,  129,  129,  535,
- /*  1200 */    30,  535,  116,  116,  116,  116,  115,  115,  114,  114,
- /*  1210 */   114,  113,  415,  535,  216, 1062, 1276,  535,  370,  535,
- /*  1220 */   130,  130,   74,   74,  535,  915,  389,  876,   17,  437,
- /*  1230 */   429,   31,  875,  916,   76,   76,  266,  101,   73,   73,
- /*  1240 */    43,   43,  835,  834,  308,   47,   47,   95,  825,  943,
- /*  1250 */   441,  938,  241,  241,  305,  443,  313,  384,  241,   95,
- /*  1260 */   842,  843,  193,  465, 1209,  327,  237,  436,   95, 1011,
- /*  1270 */  1007,  909,  873,  237,  241,  107, 1023,  384, 1023,  955,
- /*  1280 */  1415,  119,  120,  110, 1136, 1136,  981,  984,  974,  974,
- /*  1290 */   117,  117,  118,  118,  118,  118, 1022,  809, 1022,  825,
- /*  1300 */   137,  119,  108,  110, 1136, 1136,  981,  984,  974,  974,
- /*  1310 */   117,  117,  118,  118,  118,  118,  874, 1414,  451,  107,
- /*  1320 */  1011,  314, 1273,  318,  218,  321,  323,  325, 1224, 1208,
- /*  1330 */   955,  330,  339,  340,  116,  116,  116,  116,  115,  115,
- /*  1340 */   114,  114,  114,  113,  415, 1285, 1322, 1260, 1493, 1470,
- /*  1350 */  1271,  283,  521, 1328,  116,  116,  116,  116,  115,  115,
- /*  1360 */   114,  114,  114,  113,  415, 1191, 1184, 1173, 1172, 1174,
- /*  1370 */   522, 1487,  211,  460,  384,  256,  199,  367, 1257,  342,
- /*  1380 */   195,  470,  307,  344,   11,  333,  525,  445, 1307, 1315,
- /*  1390 */   375,  203, 1207, 1151,  384,  346, 1387,  188,  360,  120,
- /*  1400 */   110, 1136, 1136,  981,  984,  974,  974,  117,  117,  118,
- /*  1410 */   118,  118,  118, 1386,  428, 1490,  245,  300,  348, 1148,
- /*  1420 */   110, 1136, 1136,  981,  984,  974,  974,  117,  117,  118,
- /*  1430 */   118,  118,  118,  189,  198, 1434, 1432,   78,   81,  163,
- /*  1440 */    82,  392,  439, 1392,  173,  105,  527,   35,    4,  157,
- /*  1450 */  1312,  116,  116,  116,  116,  115,  115,  114,  114,  114,
- /*  1460 */   113,  415,  530,  165,   93, 1304,  431,  432,  168,  463,
- /*  1470 */   221,  116,  116,  116,  116,  115,  115,  114,  114,  114,
- /*  1480 */   113,  415,  169,  452,  170,  416,  171,  374,  372,  438,
- /*  1490 */    36, 1318,  177,  225, 1381,   87,  458,  524, 1403,  316,
- /*  1500 */   257,  105,  527,  227,    4,  182,  461,  160,  320,  228,
- /*  1510 */   377, 1175,  476,  229, 1227, 1226,  405, 1225,  530, 1218,
- /*  1520 */   961,  378, 1199, 1198,  827,  332,  103,  103, 1197,  407,
- /*  1530 */     8, 1217, 1502,  104,  487,  416,  537,  536,  281,  282,
- /*  1540 */   951,  416,  490, 1268,  496,   92,  341,  243, 1269,  343,
- /*  1550 */   244, 1267,  122,  524,  345, 1461,  515,  288,  526,   10,
- /*  1560 */   354, 1266, 1460,  352,  504, 1250,   99, 1367,   94,  503,
- /*  1570 */   499,  951,  951,  953,  954,   27,  961,  347, 1249,  194,
- /*  1580 */   251,  358,  103,  103,  359, 1181,   34,  538, 1110,  104,
- /*  1590 */   255,  416,  537,  536,  286,  252,  951,  254,  539,  149,
- /*  1600 */  1170, 1419, 1165, 1420, 1418,  150, 1417,  135,  279,  785,
- /*  1610 */   151,  417, 1195,  196,  290,  210,  386, 1194,  269,  387,
- /*  1620 */   162, 1021,  133,   77, 1192, 1019,  935,  951,  951,  953,
- /*  1630 */   954,   27, 1479, 1104,  418,  164,  153,  268,  217,  166,
- /*  1640 */   859,  306,  366,  366,  365,  253,  363,  220, 1035,  798,
- /*  1650 */   172,  939,  105,  527,  155,    4,  394,  174,  396,  156,
- /*  1660 */    83, 1038,  213,   84,  294,   85,   86,  223,  222,  530,
- /*  1670 */  1034,  144,  293,   18,  224,  315,  241, 1027, 1145,  178,
- /*  1680 */   457,  226,  179,   37,  800,  334,  462,  230,  328,  466,
- /*  1690 */   180,  471,  416,   88,   19,   20,   89,  280,  838,  158,
- /*  1700 */   191,   90,  215,  478,  524, 1097,  204,  192,  987,   91,
- /*  1710 */   152, 1070,   39,  154, 1071,  504,  486,   40,  489,  205,
- /*  1720 */   505,  260,  105,  527,  214,    4,  908,  961,  262,  183,
- /*  1730 */   240,   21,  903,  103,  103,  107,   22, 1086,   23,  530,
- /*  1740 */   104, 1088,  416,  537,  536,   24, 1093,  951,   25, 1074,
- /*  1750 */  1090, 1094,    7,   33,  511,  186,   26, 1002,  385,   95,
- /*  1760 */   988,  986,  416,  288,  526,  990, 1044,  246, 1043,  247,
- /*  1770 */   991,   28,   41,  106,  524,  956,  810,   29,  951,  951,
- /*  1780 */   953,  954,   27,  531,  361,  504,  423,  248,  869,  249,
- /*  1790 */   503, 1495,  364, 1105, 1161, 1494, 1161,  961, 1161, 1161,
- /*  1800 */  1161, 1161, 1161,  103,  103, 1161, 1161, 1161, 1161, 1161,
- /*  1810 */   104, 1161,  416,  537,  536, 1104,  418,  951, 1161,  268,
- /*  1820 */  1161, 1161, 1161, 1161,  366,  366,  365,  253,  363, 1161,
- /*  1830 */  1161,  798, 1161, 1161, 1161, 1161,  105,  527, 1161,    4,
- /*  1840 */  1161, 1161, 1161, 1161,  213, 1161,  294, 1161,  951,  951,
- /*  1850 */   953,  954,   27,  530,  293, 1161, 1161, 1161, 1161, 1161,
- /*  1860 */  1161, 1161, 1161, 1161, 1161, 1161, 1161, 1161, 1161, 1161,
- /*  1870 */  1161, 1161, 1161, 1161, 1161, 1161,  416, 1161, 1161, 1161,
- /*  1880 */  1161, 1161, 1161, 1161,  215, 1161, 1161, 1161,  524, 1161,
- /*  1890 */  1161, 1161,  152, 1161, 1161,  154,  105,  527, 1161,    4,
- /*  1900 */  1161, 1161, 1161, 1161, 1161, 1161,  214, 1161, 1161, 1161,
- /*  1910 */  1161,  961, 1161,  530, 1161, 1161, 1161,  103,  103,  880,
- /*  1920 */  1161, 1161, 1161, 1161,  104, 1161,  416,  537,  536, 1161,
- /*  1930 */  1161,  951, 1161, 1161, 1161, 1161,  416, 1161, 1161, 1161,
- /*  1940 */   385, 1161, 1161, 1161, 1161,  288,  526, 1161,  524, 1161,
- /*  1950 */  1161, 1161, 1161, 1161, 1161, 1161,   97,  527, 1161,    4,
- /*  1960 */  1161, 1161,  951,  951,  953,  954,   27, 1161,  423, 1161,
- /*  1970 */  1161,  961, 1161,  530, 1161, 1161, 1161,  103,  103, 1161,
- /*  1980 */  1161, 1161, 1161, 1161,  104, 1161,  416,  537,  536, 1161,
- /*  1990 */  1161,  951,  268, 1161, 1161, 1161,  416,  366,  366,  365,
- /*  2000 */   253,  363, 1161, 1161,  798, 1161, 1161, 1161,  524, 1161,
- /*  2010 */  1161, 1161, 1161, 1161, 1161, 1161, 1161,  213, 1161,  294,
- /*  2020 */  1161, 1161,  951,  951,  953,  954,   27,  293, 1161, 1161,
- /*  2030 */  1161,  961, 1161, 1161, 1161, 1161, 1161,  103,  103, 1161,
- /*  2040 */  1161, 1161, 1161, 1161,  104, 1161,  416,  537,  536, 1161,
- /*  2050 */  1161,  951, 1161, 1161, 1161, 1161, 1161,  215, 1161, 1161,
- /*  2060 */  1161, 1161, 1161, 1161, 1161,  152, 1161, 1161,  154, 1161,
- /*  2070 */  1161, 1161, 1161, 1161, 1161, 1161, 1161, 1161, 1161,  214,
- /*  2080 */  1161, 1161,  951,  951,  953,  954,   27, 1161, 1161, 1161,
- /*  2090 */  1161, 1161, 1161, 1161, 1161, 1161, 1161, 1161, 1161, 1161,
- /*  2100 */  1161, 1161, 1161, 1161, 1161, 1161, 1161, 1161, 1161, 1161,
- /*  2110 */  1161, 1161, 1161,  385, 1161, 1161, 1161, 1161,  288,  526,
- /*  2120 */  1161, 1161, 1161, 1161, 1161, 1161, 1161, 1161, 1161, 1161,
- /*  2130 */  1161, 1161, 1161, 1161, 1161, 1161, 1161, 1161, 1161, 1161,
- /*  2140 */  1161,  423,
+ /*     0 */   544, 1220,  544,  449, 1258,  544, 1237,  544,  114,  111,
+ /*    10 */   211,  544, 1535,  544, 1258,  521,  114,  111,  211,  390,
+ /*    20 */  1230,  342,   42,   42,   42,   42, 1223,   42,   42,   71,
+ /*    30 */    71,  935, 1222,   71,   71,   71,   71, 1460, 1491,  936,
+ /*    40 */   818,  451,    6,  121,  122,  112, 1163, 1163, 1004, 1007,
+ /*    50 */   997,  997,  119,  119,  120,  120,  120,  120, 1541,  390,
+ /*    60 */  1356, 1515,  550,    2, 1191,  194,  526,  434,  143,  291,
+ /*    70 */   526,  136,  526,  369,  261,  502,  272,  383, 1271,  525,
+ /*    80 */   501,  491,  164,  121,  122,  112, 1163, 1163, 1004, 1007,
+ /*    90 */   997,  997,  119,  119,  120,  120,  120,  120, 1356,  440,
+ /*   100 */  1512,  118,  118,  118,  118,  117,  117,  116,  116,  116,
+ /*   110 */   115,  422,  266,  266,  266,  266, 1496,  356, 1498,  433,
+ /*   120 */   355, 1496,  515,  522, 1483,  541, 1112,  541, 1112,  390,
+ /*   130 */   403,  241,  208,  114,  111,  211,   98,  290,  535,  221,
+ /*   140 */  1027,  118,  118,  118,  118,  117,  117,  116,  116,  116,
+ /*   150 */   115,  422, 1140,  121,  122,  112, 1163, 1163, 1004, 1007,
+ /*   160 */   997,  997,  119,  119,  120,  120,  120,  120,  404,  426,
+ /*   170 */   117,  117,  116,  116,  116,  115,  422, 1416,  466,  123,
+ /*   180 */   118,  118,  118,  118,  117,  117,  116,  116,  116,  115,
+ /*   190 */   422,  116,  116,  116,  115,  422,  538,  538,  538,  390,
+ /*   200 */   503,  120,  120,  120,  120,  113, 1049, 1140, 1141, 1142,
+ /*   210 */  1049,  118,  118,  118,  118,  117,  117,  116,  116,  116,
+ /*   220 */   115,  422, 1459,  121,  122,  112, 1163, 1163, 1004, 1007,
+ /*   230 */   997,  997,  119,  119,  120,  120,  120,  120,  390,  442,
+ /*   240 */   314,   83,  461,   81,  357,  380, 1140,   80,  118,  118,
+ /*   250 */   118,  118,  117,  117,  116,  116,  116,  115,  422,  179,
+ /*   260 */   432,  422,  121,  122,  112, 1163, 1163, 1004, 1007,  997,
+ /*   270 */   997,  119,  119,  120,  120,  120,  120,  432,  431,  266,
+ /*   280 */   266,  118,  118,  118,  118,  117,  117,  116,  116,  116,
+ /*   290 */   115,  422,  541, 1107,  901,  504, 1140,  114,  111,  211,
+ /*   300 */  1429, 1140, 1141, 1142,  206,  489, 1107,  390,  447, 1107,
+ /*   310 */   543,  328,  120,  120,  120,  120,  298, 1429, 1431,   17,
+ /*   320 */   118,  118,  118,  118,  117,  117,  116,  116,  116,  115,
+ /*   330 */   422,  121,  122,  112, 1163, 1163, 1004, 1007,  997,  997,
+ /*   340 */   119,  119,  120,  120,  120,  120,  390, 1356,  432, 1140,
+ /*   350 */   480, 1140, 1141, 1142,  994,  994, 1005, 1008,  443,  118,
+ /*   360 */   118,  118,  118,  117,  117,  116,  116,  116,  115,  422,
+ /*   370 */   121,  122,  112, 1163, 1163, 1004, 1007,  997,  997,  119,
+ /*   380 */   119,  120,  120,  120,  120, 1052, 1052,  463, 1429,  118,
+ /*   390 */   118,  118,  118,  117,  117,  116,  116,  116,  115,  422,
+ /*   400 */  1140,  449,  544, 1424, 1140, 1141, 1142,  233,  964, 1140,
+ /*   410 */   479,  476,  475,  171,  358,  390,  164,  405,  412,  840,
+ /*   420 */   474,  164,  185,  332,   71,   71, 1241,  998,  118,  118,
+ /*   430 */   118,  118,  117,  117,  116,  116,  116,  115,  422,  121,
+ /*   440 */   122,  112, 1163, 1163, 1004, 1007,  997,  997,  119,  119,
+ /*   450 */   120,  120,  120,  120,  390, 1140, 1141, 1142,  833,   12,
+ /*   460 */   313,  507,  163,  354, 1140, 1141, 1142,  114,  111,  211,
+ /*   470 */   506,  290,  535,  544,  276,  180,  290,  535,  121,  122,
+ /*   480 */   112, 1163, 1163, 1004, 1007,  997,  997,  119,  119,  120,
+ /*   490 */   120,  120,  120,  343,  482,   71,   71,  118,  118,  118,
+ /*   500 */   118,  117,  117,  116,  116,  116,  115,  422, 1140,  209,
+ /*   510 */   409,  521, 1140, 1107, 1569,  376,  252,  269,  340,  485,
+ /*   520 */   335,  484,  238,  390,  511,  362, 1107, 1125,  331, 1107,
+ /*   530 */   191,  407,  286,   32,  455,  441,  118,  118,  118,  118,
+ /*   540 */   117,  117,  116,  116,  116,  115,  422,  121,  122,  112,
+ /*   550 */  1163, 1163, 1004, 1007,  997,  997,  119,  119,  120,  120,
+ /*   560 */   120,  120,  390, 1140, 1141, 1142,  985, 1140, 1141, 1142,
+ /*   570 */  1140,  233,  490, 1490,  479,  476,  475,    6,  163,  544,
+ /*   580 */   510,  544,  115,  422,  474,    5,  121,  122,  112, 1163,
+ /*   590 */  1163, 1004, 1007,  997,  997,  119,  119,  120,  120,  120,
+ /*   600 */   120,   13,   13,   13,   13,  118,  118,  118,  118,  117,
+ /*   610 */   117,  116,  116,  116,  115,  422,  401,  500,  406,  544,
+ /*   620 */  1484,  542, 1140,  890,  890, 1140, 1141, 1142, 1471, 1140,
+ /*   630 */   275,  390,  806,  807,  808,  969,  420,  420,  420,   16,
+ /*   640 */    16,   55,   55, 1240,  118,  118,  118,  118,  117,  117,
+ /*   650 */   116,  116,  116,  115,  422,  121,  122,  112, 1163, 1163,
+ /*   660 */  1004, 1007,  997,  997,  119,  119,  120,  120,  120,  120,
+ /*   670 */   390, 1187,    1,    1,  550,    2, 1191, 1140, 1141, 1142,
+ /*   680 */   194,  291,  896,  136, 1140, 1141, 1142,  895,  519, 1490,
+ /*   690 */  1271,    3,  378,    6,  121,  122,  112, 1163, 1163, 1004,
+ /*   700 */  1007,  997,  997,  119,  119,  120,  120,  120,  120,  856,
+ /*   710 */   544,  922,  544,  118,  118,  118,  118,  117,  117,  116,
+ /*   720 */   116,  116,  115,  422,  266,  266, 1090, 1567, 1140,  549,
+ /*   730 */  1567, 1191,   13,   13,   13,   13,  291,  541,  136,  390,
+ /*   740 */   483,  419,  418,  964,  342, 1271,  466,  408,  857,  279,
+ /*   750 */   140,  221,  118,  118,  118,  118,  117,  117,  116,  116,
+ /*   760 */   116,  115,  422,  121,  122,  112, 1163, 1163, 1004, 1007,
+ /*   770 */   997,  997,  119,  119,  120,  120,  120,  120,  544,  266,
+ /*   780 */   266,  426,  390, 1140, 1141, 1142, 1170,  828, 1170,  466,
+ /*   790 */   429,  145,  541, 1144,  399,  313,  437,  301,  836, 1488,
+ /*   800 */    71,   71,  410,    6, 1088,  471,  221,  100,  112, 1163,
+ /*   810 */  1163, 1004, 1007,  997,  997,  119,  119,  120,  120,  120,
+ /*   820 */   120,  118,  118,  118,  118,  117,  117,  116,  116,  116,
+ /*   830 */   115,  422,  237, 1423,  544,  449,  426,  287,  984,  544,
+ /*   840 */   236,  235,  234,  828,   97,  527,  427, 1263, 1263, 1144,
+ /*   850 */   492,  306,  428,  836,  975,  544,   71,   71,  974, 1239,
+ /*   860 */   544,   51,   51,  300,  118,  118,  118,  118,  117,  117,
+ /*   870 */   116,  116,  116,  115,  422,  194,  103,   70,   70,  266,
+ /*   880 */   266,  544,   71,   71,  266,  266,   30,  389,  342,  974,
+ /*   890 */   974,  976,  541,  526, 1107,  326,  390,  541,  493,  395,
+ /*   900 */  1468,  195,  528,   13,   13, 1356,  240, 1107,  277,  280,
+ /*   910 */  1107,  280,  303,  455,  305,  331,  390,   31,  188,  417,
+ /*   920 */   121,  122,  112, 1163, 1163, 1004, 1007,  997,  997,  119,
+ /*   930 */   119,  120,  120,  120,  120,  142,  390,  363,  455,  984,
+ /*   940 */   121,  122,  112, 1163, 1163, 1004, 1007,  997,  997,  119,
+ /*   950 */   119,  120,  120,  120,  120,  975,  321, 1140,  324,  974,
+ /*   960 */   121,  110,  112, 1163, 1163, 1004, 1007,  997,  997,  119,
+ /*   970 */   119,  120,  120,  120,  120,  462,  375, 1183,  118,  118,
+ /*   980 */   118,  118,  117,  117,  116,  116,  116,  115,  422, 1140,
+ /*   990 */   974,  974,  976,  304,    9,  364,  244,  360,  118,  118,
+ /*  1000 */   118,  118,  117,  117,  116,  116,  116,  115,  422,  312,
+ /*  1010 */   544,  342, 1140, 1141, 1142,  299,  290,  535,  118,  118,
+ /*  1020 */   118,  118,  117,  117,  116,  116,  116,  115,  422, 1261,
+ /*  1030 */  1261, 1161,   13,   13,  278,  419,  418,  466,  390,  921,
+ /*  1040 */   260,  260,  289, 1167, 1140, 1141, 1142,  189, 1169,  266,
+ /*  1050 */   266,  466,  388,  541, 1184,  544, 1168,  263,  144,  487,
+ /*  1060 */   920,  544,  541,  122,  112, 1163, 1163, 1004, 1007,  997,
+ /*  1070 */   997,  119,  119,  120,  120,  120,  120,   71,   71, 1140,
+ /*  1080 */  1170, 1270, 1170,   13,   13,  896, 1068, 1161,  544,  466,
+ /*  1090 */   895,  107,  536, 1489,    4, 1266, 1107,    6,  523, 1047,
+ /*  1100 */    12, 1069, 1090, 1568,  311,  453, 1568,  518,  539, 1107,
+ /*  1110 */    56,   56, 1107, 1487,  421, 1356, 1070,    6,  343,  285,
+ /*  1120 */   118,  118,  118,  118,  117,  117,  116,  116,  116,  115,
+ /*  1130 */   422,  423, 1269,  319, 1140, 1141, 1142,  876,  266,  266,
+ /*  1140 */  1275,  107,  536,  533,    4, 1486,  293,  877, 1209,    6,
+ /*  1150 */   210,  541,  541,  164, 1540,  494,  414,  865,  539,  267,
+ /*  1160 */   267, 1212,  396,  509,  497,  204,  266,  266,  394,  529,
+ /*  1170 */     8,  984,  541,  517,  544,  920,  456,  105,  105,  541,
+ /*  1180 */  1088,  423,  266,  266,  106,  415,  423,  546,  545,  266,
+ /*  1190 */   266,  974,  516,  533, 1371,  541,   15,   15,  266,  266,
+ /*  1200 */   454, 1118,  541,  266,  266, 1068, 1370,  513,  290,  535,
+ /*  1210 */   544,  541,  512,   97,  442,  314,  541,  544,  920,  125,
+ /*  1220 */  1069,  984,  974,  974,  976,  977,   27,  105,  105,  399,
+ /*  1230 */   341, 1509,   44,   44,  106, 1070,  423,  546,  545,   57,
+ /*  1240 */    57,  974,  341, 1509,  107,  536,  544,    4,  460,  399,
+ /*  1250 */   214, 1118,  457,  294,  375, 1089,  532,  297,  544,  537,
+ /*  1260 */   396,  539,  290,  535,  104,  244,  102,  524,   58,   58,
+ /*  1270 */   544,  109,  974,  974,  976,  977,   27, 1514, 1129,  425,
+ /*  1280 */    59,   59,  270,  237,  423,  138,   95,  373,  373,  372,
+ /*  1290 */   255,  370,   60,   60,  815, 1178,  533,  544,  273,  544,
+ /*  1300 */  1161,  843,  387,  386,  544, 1307,  544,  215,  210,  296,
+ /*  1310 */   513,  847,  544,  265,  208,  514, 1306,  295,  274,   61,
+ /*  1320 */    61,   62,   62,  436,  984, 1160,   45,   45,   46,   46,
+ /*  1330 */   105,  105, 1184,  920,   47,   47, 1474,  106,  544,  423,
+ /*  1340 */   546,  545,  218,  544,  974,  935, 1085,  217,  544,  377,
+ /*  1350 */   395,  107,  536,  936,    4,  156, 1161,  843,  158,  544,
+ /*  1360 */    49,   49,  141,  544,   38,   50,   50,  544,  539,  307,
+ /*  1370 */    63,   63,  544, 1448,  216,  974,  974,  976,  977,   27,
+ /*  1380 */   444,   64,   64,  544, 1447,   65,   65,  544,  524,   14,
+ /*  1390 */    14,  423,  458,  544,   66,   66,  310,  544,  316,   97,
+ /*  1400 */  1034,  544,  961,  533,  268,  127,  127,  544,  391,   67,
+ /*  1410 */    67,  544,  978,  290,  535,   52,   52,  513,  544,   68,
+ /*  1420 */    68, 1294,  512,   69,   69,  397,  165,  855,  854,   53,
+ /*  1430 */    53,  984,  966,  151,  151,  243,  430,  105,  105,  199,
+ /*  1440 */   152,  152,  448, 1303,  106,  243,  423,  546,  545, 1129,
+ /*  1450 */   425,  974,  320,  270,  862,  863, 1034,  220,  373,  373,
+ /*  1460 */   372,  255,  370,  450,  323,  815,  243,  544,  978,  544,
+ /*  1470 */   107,  536,  544,    4,  544,  938,  939,  325,  215, 1046,
+ /*  1480 */   296, 1046,  974,  974,  976,  977,   27,  539,  295,   76,
+ /*  1490 */    76,   54,   54,  327,   72,   72,  128,  128, 1503, 1254,
+ /*  1500 */   107,  536,  544,    4, 1045,  544, 1045,  531, 1238,  544,
+ /*  1510 */   423,  544,  315,  334,  544,   97,  544,  539,  217,  544,
+ /*  1520 */   472, 1528,  533,  239,   73,   73,  156,  129,  129,  158,
+ /*  1530 */   467,  130,  130,  126,  126,  344,  150,  150,  149,  149,
+ /*  1540 */   423,  134,  134,  329, 1030,  216,   97,  239,  929,  345,
+ /*  1550 */   984,  243,  533, 1315,  339,  544,  105,  105,  900, 1355,
+ /*  1560 */   544, 1290,  258,  106,  338,  423,  546,  545,  544, 1301,
+ /*  1570 */   974,  893,   99,  536,  109,    4,  544,  133,  133,  391,
+ /*  1580 */   984,  197,  131,  131,  290,  535,  105,  105,  530,  539,
+ /*  1590 */   132,  132, 1361,  106, 1219,  423,  546,  545,   75,   75,
+ /*  1600 */   974,  974,  974,  976,  977,   27,  544,  430,  826, 1211,
+ /*  1610 */   894,  139,  423,  109,  544, 1200, 1199, 1201, 1522,  544,
+ /*  1620 */   201,  544,   11,  374,  533, 1287,  347,  349,   77,   77,
+ /*  1630 */  1340,  974,  974,  976,  977,   27,   74,   74,  351,  213,
+ /*  1640 */   435,   43,   43,   48,   48,  302,  477,  309, 1348,  382,
+ /*  1650 */   353,  452,  984,  337, 1237, 1420, 1419,  205,  105,  105,
+ /*  1660 */   192,  367,  193,  534, 1525,  106, 1178,  423,  546,  545,
+ /*  1670 */   247,  167,  974,  270, 1467,  200, 1465, 1175,  373,  373,
+ /*  1680 */   372,  255,  370,  398,   79,  815,   83,   82, 1425,  446,
+ /*  1690 */   161,  177,  169,   95, 1337,  438,  172,  173,  215,  174,
+ /*  1700 */   296,  175,   35,  974,  974,  976,  977,   27,  295, 1345,
+ /*  1710 */   439,  470,  223,   36,  379,  445, 1414,  381,  459, 1351,
+ /*  1720 */   181,  227,   88,  465,  259,  229, 1436,  318,  186,  468,
+ /*  1730 */   322,  230,  384, 1202,  231,  486, 1257, 1256,  217,  411,
+ /*  1740 */  1255, 1248,   90,  847,  206,  413,  156,  505, 1539,  158,
+ /*  1750 */  1226, 1538,  283, 1508, 1227,  336,  385,  284, 1225,  496,
+ /*  1760 */  1537, 1298,   94,  346,  348,  216, 1247,  499, 1299,  245,
+ /*  1770 */   246, 1297,  416,  350, 1494,  124, 1493,   10,  524,  361,
+ /*  1780 */  1400,  101,   96,  288,  508,  253, 1135, 1208,   34, 1296,
+ /*  1790 */   547,  254,  256,  257,  392,  548, 1197, 1192,  359,  391,
+ /*  1800 */  1280, 1279,  196,  365,  290,  535,  366,  352, 1452, 1322,
+ /*  1810 */  1321, 1453,  153,  137,  281,  154,  802,  424,  155, 1451,
+ /*  1820 */  1450,  198,  292,  202,  203,   78,  212,  430,  271,  135,
+ /*  1830 */  1044, 1042,  958,  168,  219,  157,  170,  879,  308,  222,
+ /*  1840 */  1058,  176,  159,  962,  400,   84,  402,  178,   85,   86,
+ /*  1850 */    87,  166,  160,  393, 1061,  224,  225, 1057,  146,   18,
+ /*  1860 */   226,  317, 1050, 1172,  243,  464,  182,  228,   37,  183,
+ /*  1870 */   817,  469,  338,  232,  330,  481,  184,   89,  845,   19,
+ /*  1880 */    20,   92,  473,  478,  333,   91,  162,  858,  147,  488,
+ /*  1890 */   282, 1123,  148, 1010,  928, 1093,   39,   93,   40,  495,
+ /*  1900 */  1094,  187,  498,  207,  262,  264,  923,  242, 1109,  109,
+ /*  1910 */  1113, 1111, 1097,   33,   21, 1117,  520, 1025,   22,   23,
+ /*  1920 */    24, 1116,   25,  190,   97, 1011, 1009,   26, 1013, 1067,
+ /*  1930 */   248,    7, 1066,  249, 1014,   28,   41,  889,  979,  827,
+ /*  1940 */   108,   29,  250,  540,  251, 1530,  371,  368, 1131, 1130,
+ /*  1950 */  1188, 1188, 1188, 1188, 1188, 1188, 1188, 1529,
 };
 static const YYCODETYPE yy_lookahead[] = {
- /*     0 */   184,  184,  259,  260,  261,  259,  260,  261,  176,  177,
- /*    10 */   178,  179,  180,  181,  184,  208,  212,  213,  186,   19,
- /*    20 */   188,  205,  206,  280,  205,  221,   22,  195,   24,  195,
- /*    30 */   208,   31,  195,  205,   29,  205,  206,  255,   33,   39,
- /*    40 */   200,  201,  202,   43,   44,   45,   46,   47,   48,   49,
- /*    50 */    50,   51,   52,   53,   54,   55,   56,   57,  205,  227,
- /*    60 */   228,  227,  228,   59,  227,  228,  259,  260,  261,  252,
- /*    70 */    65,  241,  240,  184,  240,  223,  224,  240,  244,  245,
- /*    80 */   250,  259,  260,  261,   19,  253,   54,   55,   56,   57,
- /*    90 */    58,  184,  255,  184,  205,  206,   96,   97,   98,   99,
- /*   100 */   100,  101,  102,  103,  104,  105,  106,   46,   47,   48,
- /*   110 */    49,   46,  296,  297,  110,  283,   19,   96,   97,   98,
- /*   120 */    99,  100,  101,  102,  103,  104,  105,  106,   96,   97,
- /*   130 */    98,   99,  100,  101,  102,  103,  104,  105,  106,   94,
- /*   140 */    43,   44,   45,   46,   47,   48,   49,   50,   51,   52,
- /*   150 */    53,   54,   55,   56,   57,  110,  106,   73,  251,  114,
- /*   160 */    73,  178,  179,  180,  181,   59,  184,  292,   81,  186,
- /*   170 */   295,  188,  218,  108,   19,  114,  184,   11,  195,  184,
- /*   180 */    83,  184,   85,   54,   55,   56,   57,  205,  206,  124,
- /*   190 */   145,  146,  147,   96,   97,   98,   99,  100,  101,  102,
- /*   200 */   103,  104,  105,  106,  120,  121,  122,  120,  102,   81,
- /*   210 */   227,  228,  220,   19,   16,  109,  110,  111,  131,  132,
- /*   220 */    26,  184,  184,  240,  229,   96,   97,   98,   99,  100,
- /*   230 */   101,  102,  103,  104,  105,  106,  253,   43,   44,   45,
- /*   240 */    46,   47,   48,   49,   50,   51,   52,   53,   54,   55,
- /*   250 */    56,   57,  100,  101,  102,  103,  104,  105,  106,  131,
- /*   260 */   132,  106,  127,   69,  129,  130,  283,  112,  113,  114,
- /*   270 */   115,  116,  117,  118,   81,   77,   76,   79,  296,  124,
- /*   280 */   298,  203,  184,   19,   84,   59,   86,  121,  122,   89,
- /*   290 */    96,   97,   98,   99,  100,  101,  102,  103,  104,  105,
- /*   300 */   106,  184,   35,  205,  206,   22,  113,   43,   44,   45,
- /*   310 */    46,   47,   48,   49,   50,   51,   52,   53,   54,   55,
- /*   320 */    56,   57,  205,  206,  131,  132,  100,  101,  291,  292,
- /*   330 */   114,   67,  295,   66,  108,  109,  110,  111,  138,  113,
- /*   340 */   124,   74,   59,  179,  184,  181,  184,  121,   22,  271,
- /*   350 */   186,   19,  188,  184,  276,   59,   24,  184,  241,  195,
- /*   360 */    96,   97,   98,   99,  100,  101,  102,  103,  104,  105,
- /*   370 */   106,  145,   59,  147,  184,   43,   44,   45,   46,   47,
- /*   380 */    48,   49,   50,   51,   52,   53,   54,   55,   56,   57,
- /*   390 */   123,  227,  228,  110,  296,  297,   22,   23,  184,  102,
- /*   400 */   103,  104,  105,  106,  240,  109,  110,  111,  112,  195,
- /*   410 */   204,  115,  116,  117,   22,  184,  226,  253,  212,  205,
- /*   420 */   206,  125,  109,  110,  111,   22,  100,  101,   96,   97,
- /*   430 */    98,   99,  100,  101,  102,  103,  104,  105,  106,  184,
- /*   440 */    59,  227,  228,  121,  122,   59,  277,  283,   19,  289,
- /*   450 */   290,   59,   23,   76,  240,  241,  143,   76,   72,  189,
- /*   460 */   205,  206,   59,   86,  250,   84,   89,   86,  203,   95,
- /*   470 */    89,  281,   43,   44,   45,   46,   47,   48,   49,   50,
- /*   480 */    51,   52,   53,   54,   55,   56,   57,  227,  228,  184,
- /*   490 */   109,  110,  111,   12,  184,  109,  110,  111,  184,  184,
- /*   500 */   240,  109,  110,  111,  184,  195,  214,   59,   27,  184,
- /*   510 */   205,  206,  109,  110,  111,  205,  206,  184,  263,  138,
- /*   520 */   205,  206,  184,   42,   22,   96,   97,   98,   99,  100,
- /*   530 */   101,  102,  103,  104,  105,  106,  266,  227,  228,   59,
- /*   540 */   270,  276,   94,   66,   63,   19,  241,   22,   26,   23,
- /*   550 */   240,  241,   72,   59,   73,  250,  241,  109,  110,   82,
- /*   560 */    22,   59,  114,  223,  224,  250,  252,   59,   91,   43,
- /*   570 */    44,   45,   46,   47,   48,   49,   50,   51,   52,   53,
- /*   580 */    54,   55,   56,   57,   59,  184,   26,   59,  268,  109,
- /*   590 */   110,  111,  184,  145,  146,  147,  112,   59,  203,  115,
- /*   600 */   116,  117,  277,  109,  110,  111,  205,  206,  195,  125,
- /*   610 */   277,  109,  110,  111,  100,  101,  139,  109,  110,  111,
- /*   620 */   219,  184,   96,   97,   98,   99,  100,  101,  102,  103,
- /*   630 */   104,  105,  106,  111,  109,  110,  111,  109,  110,  111,
- /*   640 */   227,  228,   19,  184,  136,  184,   23,  109,  110,  111,
- /*   650 */   200,  201,  202,  240,  259,  260,  261,  195,  136,  145,
- /*   660 */   184,  147,  184,  184,  136,  214,   43,   44,   45,   46,
- /*   670 */    47,   48,   49,   50,   51,   52,   53,   54,   55,   56,
- /*   680 */    57,  205,  206,  205,  206,  227,  228,  184,  229,  227,
- /*   690 */   228,  131,  132,  184,   59,  219,  184,  219,  240,  291,
- /*   700 */   292,  184,  240,  295,  105,  106,   22,   23,  205,  206,
- /*   710 */    26,  184,  251,  184,  205,  206,  184,  205,  206,   96,
- /*   720 */    97,   98,   99,  100,  101,  102,  103,  104,  105,  106,
- /*   730 */   251,  219,  205,  206,  205,  206,  184,  205,  206,   19,
- /*   740 */   184,   16,  184,   23,  241,  110,  219,   21,  219,  184,
- /*   750 */   241,  219,  286,  287,  195,  184,  195,  205,  206,  201,
- /*   760 */   202,  205,  206,   43,   44,   45,   46,   47,   48,   49,
- /*   770 */    50,   51,   52,   53,   54,   55,   56,   57,  184,   95,
- /*   780 */    22,   23,  184,   26,   26,  220,  227,  228,  227,  228,
- /*   790 */   196,  184,   23,  241,   26,   26,  195,  241,  184,  240,
- /*   800 */    12,  240,   77,   26,   79,  195,   80,  290,  201,  202,
- /*   810 */   216,  184,  218,  195,  184,   27,   96,   97,   98,   99,
- /*   820 */   100,  101,  102,  103,  104,  105,  106,  269,  227,  228,
- /*   830 */    42,  184,  205,  206,  184,  184,   19,  227,  228,  192,
- /*   840 */    23,  240,  116,  196,   76,  227,  228,  120,  121,  122,
- /*   850 */   240,   63,  254,   95,   86,  205,  206,   89,  240,  184,
- /*   860 */    43,   44,   45,   46,   47,   48,   49,   50,   51,   52,
- /*   870 */    53,   54,   55,   56,   57,  184,  269,  184,  153,  153,
- /*   880 */   111,  184,    7,    8,    9,  184,  138,  184,  184,  196,
- /*   890 */   184,  120,  121,  122,  184,  138,  205,  206,  184,  102,
- /*   900 */   184,  184,  205,  206,  156,  136,  205,  206,  205,  206,
- /*   910 */   198,  199,  135,   96,   97,   98,   99,  100,  101,  102,
- /*   920 */   103,  104,  105,  106,  184,  128,  184,  184,  184,  254,
- /*   930 */   133,  184,  237,   19,  239,  229,  226,   23,  292,  184,
- /*   940 */   226,  295,  226,  226,  184,  205,  206,  205,  206,  205,
- /*   950 */   206,  184,  292,   19,  184,  295,  252,   43,   44,   45,
- /*   960 */    46,   47,   48,   49,   50,   51,   52,   53,   54,   55,
- /*   970 */    56,   57,  205,  206,  184,  205,  206,   43,   44,   45,
- /*   980 */    46,   47,   48,   49,   50,   51,   52,   53,   54,   55,
- /*   990 */    56,   57,  157,  158,   26,  205,  206,  254,   26,  252,
- /*  1000 */   184,   15,  184,  184,  184,  292,  184,  252,  295,   24,
- /*  1010 */    96,   97,   98,   99,  100,  101,  102,  103,  104,  105,
- /*  1020 */   106,  205,  206,  205,  206,  205,  206,  205,  206,  184,
- /*  1030 */    96,   97,   98,   99,  100,  101,  102,  103,  104,  105,
- /*  1040 */   106,  184,  184,  184,   59,  184,   60,  184,  229,  184,
- /*  1050 */   205,  206,  184,  258,  184,   19,  184,   19,  184,  246,
- /*  1060 */   184,  258,  205,  206,  205,  206,  205,  206,  205,  206,
- /*  1070 */   205,  206,  184,  205,  206,  205,  206,  205,  206,  205,
- /*  1080 */   206,  205,  206,  292,  226,  151,  295,  184,  228,  294,
- /*  1090 */   184,  119,  184,  205,  206,  110,  150,  294,  152,  184,
- /*  1100 */   240,  184,   22,   23,   23,   19,  184,   26,  205,  206,
- /*  1110 */   142,  205,  206,  205,  206,  184,  198,  199,  131,  132,
- /*  1120 */   205,  206,  205,  206,   22,   19,   24,  205,  206,   43,
- /*  1130 */    44,   45,   46,   47,   48,   49,   50,   51,   52,   53,
- /*  1140 */    54,   55,   56,   57,  184,  109,  184,  109,  184,   43,
- /*  1150 */    44,   45,   46,   47,   48,   49,   50,   51,   52,   53,
- /*  1160 */    54,   55,   56,   57,   46,  205,  206,  205,  206,  205,
- /*  1170 */   206,  232,  184,  184,  184,   95,  184,  284,  285,  244,
- /*  1180 */   245,  242,   96,   97,   98,   99,  100,  101,  102,  103,
- /*  1190 */   104,  105,  106,  205,  206,  205,  206,  205,  206,  184,
- /*  1200 */    22,  184,   96,   97,   98,   99,  100,  101,  102,  103,
- /*  1210 */   104,  105,  106,  184,   24,   23,  184,  184,   26,  184,
- /*  1220 */   205,  206,  205,  206,  184,   31,  108,  128,   22,  122,
- /*  1230 */   184,   53,  133,   39,  205,  206,   22,  151,  205,  206,
- /*  1240 */   205,  206,  113,  114,   23,  205,  206,   26,   59,   23,
- /*  1250 */    23,  144,   26,   26,  184,   23,   23,   19,   26,   26,
- /*  1260 */     7,    8,   24,   23,  214,   23,   26,   61,   26,   59,
- /*  1270 */    23,   23,   23,   26,   26,   26,  145,   19,  147,   59,
- /*  1280 */   184,   43,   44,   45,   46,   47,   48,   49,   50,   51,
- /*  1290 */    52,   53,   54,   55,   56,   57,  145,   23,  147,  110,
- /*  1300 */    26,   43,   44,   45,   46,   47,   48,   49,   50,   51,
- /*  1310 */    52,   53,   54,   55,   56,   57,   23,  184,  184,   26,
- /*  1320 */   110,  184,  184,  184,  134,  184,  184,  184,  184,  184,
- /*  1330 */   110,  184,  184,  184,   96,   97,   98,   99,  100,  101,
- /*  1340 */   102,  103,  104,  105,  106,  184,  184,  184,  134,  300,
- /*  1350 */   184,  243,  184,  184,   96,   97,   98,   99,  100,  101,
- /*  1360 */   102,  103,  104,  105,  106,  184,  184,  184,  184,  184,
- /*  1370 */   224,  184,  282,  273,   19,  272,  203,  182,  243,  243,
- /*  1380 */   230,  209,  278,  243,  231,  208,  265,  278,  234,  234,
- /*  1390 */   234,  217,  213,   60,   19,  243,  208,  237,  233,   44,
- /*  1400 */    45,   46,   47,   48,   49,   50,   51,   52,   53,   54,
- /*  1410 */    55,   56,   57,  208,  247,  187,  134,  247,  247,   38,
- /*  1420 */    45,   46,   47,   48,   49,   50,   51,   52,   53,   54,
- /*  1430 */    55,   56,   57,  237,  231,  191,  191,  279,  279,  282,
- /*  1440 */   143,  191,  108,  268,   22,   19,   20,  256,   22,   43,
- /*  1450 */   257,   96,   97,   98,   99,  100,  101,  102,  103,  104,
- /*  1460 */   105,  106,   36,  222,  142,  234,   18,  191,  225,   18,
- /*  1470 */   190,   96,   97,   98,   99,  100,  101,  102,  103,  104,
- /*  1480 */   105,  106,  225,  191,  225,   59,  225,  257,  234,  234,
- /*  1490 */   256,  222,  222,  190,  234,  150,   62,   71,  275,  274,
- /*  1500 */   191,   19,   20,  190,   22,   22,  210,   81,  191,  190,
- /*  1510 */   210,  191,  108,  190,  207,  207,   64,  207,   36,  215,
- /*  1520 */    94,  210,  207,  209,  119,  207,  100,  101,  207,  106,
- /*  1530 */    48,  215,  207,  107,  210,  109,  110,  111,  267,  267,
- /*  1540 */   114,   59,  210,  249,  137,  108,  248,  191,  249,  248,
- /*  1550 */    88,  249,  141,   71,  248,  299,  138,  131,  132,   22,
- /*  1560 */   191,  249,  299,  237,   82,  238,  150,  262,  140,   87,
- /*  1570 */   139,  145,  146,  147,  148,  149,   94,  248,  238,  236,
- /*  1580 */    25,  235,  100,  101,  234,  194,   26,  193,   13,  107,
- /*  1590 */     6,  109,  110,  111,  264,  185,  114,  185,  183,  197,
- /*  1600 */   183,  203,  183,  203,  203,  197,  203,  211,  211,    4,
- /*  1610 */   197,    3,  203,   22,  155,   15,  288,  203,   93,  288,
- /*  1620 */   285,   23,   16,  203,  203,   23,  132,  145,  146,  147,
- /*  1630 */   148,  149,    0,    1,    2,  143,  123,    5,   24,  135,
- /*  1640 */    20,   16,   10,   11,   12,   13,   14,  137,    1,   17,
- /*  1650 */   135,  144,   19,   20,  123,   22,   61,  143,   37,  123,
- /*  1660 */    53,  109,   30,   53,   32,   53,   53,  134,   34,   36,
- /*  1670 */     1,    5,   40,   22,  108,  153,   26,   68,   75,   68,
- /*  1680 */    41,  134,  108,   24,   20,  124,   19,  118,   23,   67,
- /*  1690 */    22,   67,   59,   22,   22,   22,   22,   67,   28,   37,
- /*  1700 */    23,  142,   70,   22,   71,   23,  157,   23,   23,   26,
- /*  1710 */    78,   23,   22,   81,   23,   82,   24,   22,   24,  134,
- /*  1720 */    87,   23,   19,   20,   92,   22,  109,   94,   23,   22,
- /*  1730 */    34,   34,  136,  100,  101,   26,   34,   85,   34,   36,
- /*  1740 */   107,   83,  109,  110,  111,   34,   90,  114,   34,   23,
- /*  1750 */    75,   75,   44,   22,   24,   26,   34,   23,  126,   26,
- /*  1760 */    23,   23,   59,  131,  132,   23,   23,   26,   23,   22,
- /*  1770 */    11,   22,   22,   22,   71,   23,   23,   22,  145,  146,
- /*  1780 */   147,  148,  149,   26,   23,   82,  154,  134,  128,  134,
- /*  1790 */    87,  134,   15,    1,  301,  134,  301,   94,  301,  301,
- /*  1800 */   301,  301,  301,  100,  101,  301,  301,  301,  301,  301,
- /*  1810 */   107,  301,  109,  110,  111,    1,    2,  114,  301,    5,
- /*  1820 */   301,  301,  301,  301,   10,   11,   12,   13,   14,  301,
- /*  1830 */   301,   17,  301,  301,  301,  301,   19,   20,  301,   22,
- /*  1840 */   301,  301,  301,  301,   30,  301,   32,  301,  145,  146,
- /*  1850 */   147,  148,  149,   36,   40,  301,  301,  301,  301,  301,
- /*  1860 */   301,  301,  301,  301,  301,  301,  301,  301,  301,  301,
- /*  1870 */   301,  301,  301,  301,  301,  301,   59,  301,  301,  301,
- /*  1880 */   301,  301,  301,  301,   70,  301,  301,  301,   71,  301,
- /*  1890 */   301,  301,   78,  301,  301,   81,   19,   20,  301,   22,
- /*  1900 */   301,  301,  301,  301,  301,  301,   92,  301,  301,  301,
- /*  1910 */   301,   94,  301,   36,  301,  301,  301,  100,  101,  102,
- /*  1920 */   301,  301,  301,  301,  107,  301,  109,  110,  111,  301,
- /*  1930 */   301,  114,  301,  301,  301,  301,   59,  301,  301,  301,
- /*  1940 */   126,  301,  301,  301,  301,  131,  132,  301,   71,  301,
- /*  1950 */   301,  301,  301,  301,  301,  301,   19,   20,  301,   22,
- /*  1960 */   301,  301,  145,  146,  147,  148,  149,  301,  154,  301,
- /*  1970 */   301,   94,  301,   36,  301,  301,  301,  100,  101,  301,
- /*  1980 */   301,  301,  301,  301,  107,  301,  109,  110,  111,  301,
- /*  1990 */   301,  114,    5,  301,  301,  301,   59,   10,   11,   12,
- /*  2000 */    13,   14,  301,  301,   17,  301,  301,  301,   71,  301,
- /*  2010 */   301,  301,  301,  301,  301,  301,  301,   30,  301,   32,
- /*  2020 */   301,  301,  145,  146,  147,  148,  149,   40,  301,  301,
- /*  2030 */   301,   94,  301,  301,  301,  301,  301,  100,  101,  301,
- /*  2040 */   301,  301,  301,  301,  107,  301,  109,  110,  111,  301,
- /*  2050 */   301,  114,  301,  301,  301,  301,  301,   70,  301,  301,
- /*  2060 */   301,  301,  301,  301,  301,   78,  301,  301,   81,  301,
- /*  2070 */   301,  301,  301,  301,  301,  301,  301,  301,  301,   92,
- /*  2080 */   301,  301,  145,  146,  147,  148,  149,  301,  301,  301,
- /*  2090 */   301,  301,  301,  301,  301,  301,  301,  301,  301,  301,
- /*  2100 */   301,  301,  301,  301,  301,  301,  301,  301,  301,  301,
- /*  2110 */   301,  301,  301,  126,  301,  301,  301,  301,  131,  132,
- /*  2120 */   301,  301,  301,  301,  301,  301,  301,  301,  301,  301,
- /*  2130 */   301,  301,  301,  301,  301,  301,  301,  301,  301,  301,
- /*  2140 */   301,  154,  301,  301,  301,  301,  301,  301,  301,  301,
- /*  2150 */   301,  301,  301,  301,  301,  301,  301,  301,  301,  301,
- /*  2160 */   301,  301,  301,  301,  301,  301,  301,  301,  301,
+ /*     0 */   189,  211,  189,  189,  218,  189,  220,  189,  267,  268,
+ /*    10 */   269,  189,  210,  189,  228,  189,  267,  268,  269,   19,
+ /*    20 */   218,  189,  211,  212,  211,  212,  211,  211,  212,  211,
+ /*    30 */   212,   31,  211,  211,  212,  211,  212,  288,  300,   39,
+ /*    40 */    21,  189,  304,   43,   44,   45,   46,   47,   48,   49,
+ /*    50 */    50,   51,   52,   53,   54,   55,   56,   57,  225,   19,
+ /*    60 */   189,  183,  184,  185,  186,  189,  248,  263,  236,  191,
+ /*    70 */   248,  193,  248,  197,  208,  257,  262,  201,  200,  257,
+ /*    80 */   200,  257,   81,   43,   44,   45,   46,   47,   48,   49,
+ /*    90 */    50,   51,   52,   53,   54,   55,   56,   57,  189,   80,
+ /*   100 */   189,  101,  102,  103,  104,  105,  106,  107,  108,  109,
+ /*   110 */   110,  111,  234,  235,  234,  235,  305,  306,  305,  118,
+ /*   120 */   307,  305,  306,  297,  298,  247,   86,  247,   88,   19,
+ /*   130 */   259,  251,  252,  267,  268,  269,   26,  136,  137,  261,
+ /*   140 */   121,  101,  102,  103,  104,  105,  106,  107,  108,  109,
+ /*   150 */   110,  111,   59,   43,   44,   45,   46,   47,   48,   49,
+ /*   160 */    50,   51,   52,   53,   54,   55,   56,   57,  259,  291,
+ /*   170 */   105,  106,  107,  108,  109,  110,  111,  158,  189,   69,
+ /*   180 */   101,  102,  103,  104,  105,  106,  107,  108,  109,  110,
+ /*   190 */   111,  107,  108,  109,  110,  111,  205,  206,  207,   19,
+ /*   200 */    19,   54,   55,   56,   57,   58,   29,  114,  115,  116,
+ /*   210 */    33,  101,  102,  103,  104,  105,  106,  107,  108,  109,
+ /*   220 */   110,  111,  233,   43,   44,   45,   46,   47,   48,   49,
+ /*   230 */    50,   51,   52,   53,   54,   55,   56,   57,   19,  126,
+ /*   240 */   127,  148,   65,   24,  214,  200,   59,   67,  101,  102,
+ /*   250 */   103,  104,  105,  106,  107,  108,  109,  110,  111,   22,
+ /*   260 */   189,  111,   43,   44,   45,   46,   47,   48,   49,   50,
+ /*   270 */    51,   52,   53,   54,   55,   56,   57,  206,  207,  234,
+ /*   280 */   235,  101,  102,  103,  104,  105,  106,  107,  108,  109,
+ /*   290 */   110,  111,  247,   76,  107,  114,   59,  267,  268,  269,
+ /*   300 */   189,  114,  115,  116,  162,  163,   89,   19,  263,   92,
+ /*   310 */   189,   23,   54,   55,   56,   57,  189,  206,  207,   22,
+ /*   320 */   101,  102,  103,  104,  105,  106,  107,  108,  109,  110,
+ /*   330 */   111,   43,   44,   45,   46,   47,   48,   49,   50,   51,
+ /*   340 */    52,   53,   54,   55,   56,   57,   19,  189,  277,   59,
+ /*   350 */    23,  114,  115,  116,   46,   47,   48,   49,   61,  101,
+ /*   360 */   102,  103,  104,  105,  106,  107,  108,  109,  110,  111,
+ /*   370 */    43,   44,   45,   46,   47,   48,   49,   50,   51,   52,
+ /*   380 */    53,   54,   55,   56,   57,  125,  126,  127,  277,  101,
+ /*   390 */   102,  103,  104,  105,  106,  107,  108,  109,  110,  111,
+ /*   400 */    59,  189,  189,  276,  114,  115,  116,  117,   73,   59,
+ /*   410 */   120,  121,  122,   72,  214,   19,   81,  259,   19,   23,
+ /*   420 */   130,   81,   72,   24,  211,  212,  221,  119,  101,  102,
+ /*   430 */   103,  104,  105,  106,  107,  108,  109,  110,  111,   43,
+ /*   440 */    44,   45,   46,   47,   48,   49,   50,   51,   52,   53,
+ /*   450 */    54,   55,   56,   57,   19,  114,  115,  116,   23,  208,
+ /*   460 */   125,  248,  189,  189,  114,  115,  116,  267,  268,  269,
+ /*   470 */   189,  136,  137,  189,  262,   22,  136,  137,   43,   44,
+ /*   480 */    45,   46,   47,   48,   49,   50,   51,   52,   53,   54,
+ /*   490 */    55,   56,   57,  189,   95,  211,  212,  101,  102,  103,
+ /*   500 */   104,  105,  106,  107,  108,  109,  110,  111,   59,  189,
+ /*   510 */   111,  189,   59,   76,  294,  295,  117,  118,  119,  120,
+ /*   520 */   121,  122,  123,   19,   87,  189,   89,   23,  129,   92,
+ /*   530 */   279,  227,  248,   22,  189,  284,  101,  102,  103,  104,
+ /*   540 */   105,  106,  107,  108,  109,  110,  111,   43,   44,   45,
+ /*   550 */    46,   47,   48,   49,   50,   51,   52,   53,   54,   55,
+ /*   560 */    56,   57,   19,  114,  115,  116,   23,  114,  115,  116,
+ /*   570 */    59,  117,  299,  300,  120,  121,  122,  304,  189,  189,
+ /*   580 */   143,  189,  110,  111,  130,   22,   43,   44,   45,   46,
+ /*   590 */    47,   48,   49,   50,   51,   52,   53,   54,   55,   56,
+ /*   600 */    57,  211,  212,  211,  212,  101,  102,  103,  104,  105,
+ /*   610 */   106,  107,  108,  109,  110,  111,  226,  189,  226,  189,
+ /*   620 */   298,  132,   59,  134,  135,  114,  115,  116,  189,   59,
+ /*   630 */   285,   19,    7,    8,    9,   23,  205,  206,  207,  211,
+ /*   640 */   212,  211,  212,  221,  101,  102,  103,  104,  105,  106,
+ /*   650 */   107,  108,  109,  110,  111,   43,   44,   45,   46,   47,
+ /*   660 */    48,   49,   50,   51,   52,   53,   54,   55,   56,   57,
+ /*   670 */    19,  181,  182,  183,  184,  185,  186,  114,  115,  116,
+ /*   680 */   189,  191,  133,  193,  114,  115,  116,  138,  299,  300,
+ /*   690 */   200,   22,  201,  304,   43,   44,   45,   46,   47,   48,
+ /*   700 */    49,   50,   51,   52,   53,   54,   55,   56,   57,   35,
+ /*   710 */   189,  141,  189,  101,  102,  103,  104,  105,  106,  107,
+ /*   720 */   108,  109,  110,  111,  234,  235,   22,   23,   59,  184,
+ /*   730 */    26,  186,  211,  212,  211,  212,  191,  247,  193,   19,
+ /*   740 */    66,  105,  106,   73,  189,  200,  189,  226,   74,  226,
+ /*   750 */    22,  261,  101,  102,  103,  104,  105,  106,  107,  108,
+ /*   760 */   109,  110,  111,   43,   44,   45,   46,   47,   48,   49,
+ /*   770 */    50,   51,   52,   53,   54,   55,   56,   57,  189,  234,
+ /*   780 */   235,  291,   19,  114,  115,  116,  150,   59,  152,  189,
+ /*   790 */   233,  236,  247,   59,  189,  125,  126,  127,   59,  300,
+ /*   800 */   211,  212,  128,  304,  100,   19,  261,  156,   45,   46,
+ /*   810 */    47,   48,   49,   50,   51,   52,   53,   54,   55,   56,
+ /*   820 */    57,  101,  102,  103,  104,  105,  106,  107,  108,  109,
+ /*   830 */   110,  111,   46,  233,  189,  189,  291,  248,   99,  189,
+ /*   840 */   125,  126,  127,  115,   26,  200,  289,  230,  231,  115,
+ /*   850 */   200,   16,  189,  114,  115,  189,  211,  212,  119,  221,
+ /*   860 */   189,  211,  212,  258,  101,  102,  103,  104,  105,  106,
+ /*   870 */   107,  108,  109,  110,  111,  189,  156,  211,  212,  234,
+ /*   880 */   235,  189,  211,  212,  234,  235,   22,  201,  189,  150,
+ /*   890 */   151,  152,  247,  248,   76,   16,   19,  247,  248,  113,
+ /*   900 */   189,   24,  257,  211,  212,  189,   26,   89,  262,  223,
+ /*   910 */    92,  225,   77,  189,   79,  129,   19,   53,  226,  248,
+ /*   920 */    43,   44,   45,   46,   47,   48,   49,   50,   51,   52,
+ /*   930 */    53,   54,   55,   56,   57,  236,   19,  271,  189,   99,
+ /*   940 */    43,   44,   45,   46,   47,   48,   49,   50,   51,   52,
+ /*   950 */    53,   54,   55,   56,   57,  115,   77,   59,   79,  119,
+ /*   960 */    43,   44,   45,   46,   47,   48,   49,   50,   51,   52,
+ /*   970 */    53,   54,   55,   56,   57,  259,   22,   23,  101,  102,
+ /*   980 */   103,  104,  105,  106,  107,  108,  109,  110,  111,   59,
+ /*   990 */   150,  151,  152,  158,   22,  244,   24,  246,  101,  102,
+ /*  1000 */   103,  104,  105,  106,  107,  108,  109,  110,  111,  285,
+ /*  1010 */   189,  189,  114,  115,  116,  200,  136,  137,  101,  102,
+ /*  1020 */   103,  104,  105,  106,  107,  108,  109,  110,  111,  230,
+ /*  1030 */   231,   59,  211,  212,  285,  105,  106,  189,   19,  141,
+ /*  1040 */   234,  235,  239,  113,  114,  115,  116,  226,  118,  234,
+ /*  1050 */   235,  189,  249,  247,  100,  189,  126,   23,  236,  107,
+ /*  1060 */    26,  189,  247,   44,   45,   46,   47,   48,   49,   50,
+ /*  1070 */    51,   52,   53,   54,   55,   56,   57,  211,  212,   59,
+ /*  1080 */   150,  233,  152,  211,  212,  133,   12,  115,  189,  189,
+ /*  1090 */   138,   19,   20,  300,   22,  233,   76,  304,  226,   11,
+ /*  1100 */   208,   27,   22,   23,  200,   19,   26,   87,   36,   89,
+ /*  1110 */   211,  212,   92,  300,  248,  189,   42,  304,  189,  250,
+ /*  1120 */   101,  102,  103,  104,  105,  106,  107,  108,  109,  110,
+ /*  1130 */   111,   59,  200,  233,  114,  115,  116,   63,  234,  235,
+ /*  1140 */   235,   19,   20,   71,   22,  300,  189,   73,  200,  304,
+ /*  1150 */   116,  247,  247,   81,   23,  200,  227,   26,   36,  234,
+ /*  1160 */   235,  203,  204,  143,  200,   26,  234,  235,  194,  200,
+ /*  1170 */    48,   99,  247,   66,  189,  141,  284,  105,  106,  247,
+ /*  1180 */   100,   59,  234,  235,  112,  259,  114,  115,  116,  234,
+ /*  1190 */   235,  119,   85,   71,  266,  247,  211,  212,  234,  235,
+ /*  1200 */   114,   94,  247,  234,  235,   12,  266,   85,  136,  137,
+ /*  1210 */   189,  247,   90,   26,  126,  127,  247,  189,   26,   22,
+ /*  1220 */    27,   99,  150,  151,  152,  153,  154,  105,  106,  189,
+ /*  1230 */   302,  303,  211,  212,  112,   42,  114,  115,  116,  211,
+ /*  1240 */   212,  119,  302,  303,   19,   20,  189,   22,  274,  189,
+ /*  1250 */    15,  144,  278,  189,   22,   23,   63,  189,  189,  203,
+ /*  1260 */   204,   36,  136,  137,  155,   24,  157,  143,  211,  212,
+ /*  1270 */   189,   26,  150,  151,  152,  153,  154,    0,    1,    2,
+ /*  1280 */   211,  212,    5,   46,   59,  161,  147,   10,   11,   12,
+ /*  1290 */    13,   14,  211,  212,   17,   60,   71,  189,  258,  189,
+ /*  1300 */    59,   59,  105,  106,  189,  189,  189,   30,  116,   32,
+ /*  1310 */    85,  124,  189,  251,  252,   90,  189,   40,  258,  211,
+ /*  1320 */   212,  211,  212,  189,   99,   26,  211,  212,  211,  212,
+ /*  1330 */   105,  106,  100,  141,  211,  212,  189,  112,  189,  114,
+ /*  1340 */   115,  116,   24,  189,  119,   31,   23,   70,  189,   26,
+ /*  1350 */   113,   19,   20,   39,   22,   78,  115,  115,   81,  189,
+ /*  1360 */   211,  212,   22,  189,   24,  211,  212,  189,   36,  189,
+ /*  1370 */   211,  212,  189,  189,   97,  150,  151,  152,  153,  154,
+ /*  1380 */   127,  211,  212,  189,  189,  211,  212,  189,  143,  211,
+ /*  1390 */   212,   59,  189,  189,  211,  212,   23,  189,  189,   26,
+ /*  1400 */    59,  189,  149,   71,   22,  211,  212,  189,  131,  211,
+ /*  1410 */   212,  189,   59,  136,  137,  211,  212,   85,  189,  211,
+ /*  1420 */   212,  253,   90,  211,  212,  292,  293,  118,  119,  211,
+ /*  1430 */   212,   99,   23,  211,  212,   26,  159,  105,  106,  140,
+ /*  1440 */   211,  212,   23,  189,  112,   26,  114,  115,  116,    1,
+ /*  1450 */     2,  119,  189,    5,    7,    8,  115,  139,   10,   11,
+ /*  1460 */    12,   13,   14,   23,  189,   17,   26,  189,  115,  189,
+ /*  1470 */    19,   20,  189,   22,  189,   83,   84,  189,   30,  150,
+ /*  1480 */    32,  152,  150,  151,  152,  153,  154,   36,   40,  211,
+ /*  1490 */   212,  211,  212,  189,  211,  212,  211,  212,  309,  189,
+ /*  1500 */    19,   20,  189,   22,  150,  189,  152,  231,  189,  189,
+ /*  1510 */    59,  189,   23,  189,  189,   26,  189,   36,   70,  189,
+ /*  1520 */    23,  139,   71,   26,  211,  212,   78,  211,  212,   81,
+ /*  1530 */   281,  211,  212,  211,  212,  189,  211,  212,  211,  212,
+ /*  1540 */    59,  211,  212,   23,   23,   97,   26,   26,   23,  189,
+ /*  1550 */    99,   26,   71,  189,  119,  189,  105,  106,  107,  189,
+ /*  1560 */   189,  189,  280,  112,  129,  114,  115,  116,  189,  189,
+ /*  1570 */   119,   23,   19,   20,   26,   22,  189,  211,  212,  131,
+ /*  1580 */    99,  237,  211,  212,  136,  137,  105,  106,  189,   36,
+ /*  1590 */   211,  212,  189,  112,  189,  114,  115,  116,  211,  212,
+ /*  1600 */   119,  150,  151,  152,  153,  154,  189,  159,   23,  189,
+ /*  1610 */    23,   26,   59,   26,  189,  189,  189,  189,  189,  189,
+ /*  1620 */   209,  189,  238,  187,   71,  250,  250,  250,  211,  212,
+ /*  1630 */   241,  150,  151,  152,  153,  154,  211,  212,  250,  290,
+ /*  1640 */   254,  211,  212,  211,  212,  254,  215,  286,  241,  241,
+ /*  1650 */   254,  286,   99,  214,  220,  214,  214,  224,  105,  106,
+ /*  1660 */   244,  240,  244,  273,  192,  112,   60,  114,  115,  116,
+ /*  1670 */   139,  290,  119,    5,  196,  238,  196,   38,   10,   11,
+ /*  1680 */    12,   13,   14,  196,  287,   17,  148,  287,  276,  113,
+ /*  1690 */    43,   22,  229,  147,  241,   18,  232,  232,   30,  232,
+ /*  1700 */    32,  232,  264,  150,  151,  152,  153,  154,   40,  265,
+ /*  1710 */   196,   18,  195,  264,  241,  241,  241,  265,  196,  229,
+ /*  1720 */   229,  195,  155,   62,  196,  195,  283,  282,   22,  216,
+ /*  1730 */   196,  195,  216,  196,  195,  113,  213,  213,   70,   64,
+ /*  1740 */   213,  222,   22,  124,  162,  111,   78,  142,  219,   81,
+ /*  1750 */   215,  219,  275,  303,  213,  213,  216,  275,  213,  216,
+ /*  1760 */   213,  256,  113,  255,  255,   97,  222,  216,  256,  196,
+ /*  1770 */    91,  256,   82,  255,  308,  146,  308,   22,  143,  196,
+ /*  1780 */   270,  155,  145,  272,  144,   25,   13,  199,   26,  256,
+ /*  1790 */   198,  190,  190,    6,  296,  188,  188,  188,  244,  131,
+ /*  1800 */   245,  245,  243,  242,  136,  137,  241,  255,  208,  260,
+ /*  1810 */   260,  208,  202,  217,  217,  202,    4,    3,  202,  208,
+ /*  1820 */   208,   22,  160,  209,  209,  208,   15,  159,   98,   16,
+ /*  1830 */    23,   23,  137,  148,   24,  128,  140,   20,   16,  142,
+ /*  1840 */     1,  140,  128,  149,   61,   53,   37,  148,   53,   53,
+ /*  1850 */    53,  293,  128,  296,  114,   34,  139,    1,    5,   22,
+ /*  1860 */   113,  158,   68,   75,   26,   41,   68,  139,   24,  113,
+ /*  1870 */    20,   19,  129,  123,   23,   96,   22,   22,   59,   22,
+ /*  1880 */    22,  147,   67,   67,   24,   22,   37,   28,   23,   22,
+ /*  1890 */    67,   23,   23,   23,  114,   23,   22,   26,   22,   24,
+ /*  1900 */    23,   22,   24,  139,   23,   23,  141,   34,   88,   26,
+ /*  1910 */    75,   86,   23,   22,   34,   75,   24,   23,   34,   34,
+ /*  1920 */    34,   93,   34,   26,   26,   23,   23,   34,   23,   23,
+ /*  1930 */    26,   44,   23,   22,   11,   22,   22,  133,   23,   23,
+ /*  1940 */    22,   22,  139,   26,  139,  139,   15,   23,    1,    1,
+ /*  1950 */   310,  310,  310,  310,  310,  310,  310,  139,  310,  310,
+ /*  1960 */   310,  310,  310,  310,  310,  310,  310,  310,  310,  310,
+ /*  1970 */   310,  310,  310,  310,  310,  310,  310,  310,  310,  310,
+ /*  1980 */   310,  310,  310,  310,  310,  310,  310,  310,  310,  310,
+ /*  1990 */   310,  310,  310,  310,  310,  310,  310,  310,  310,  310,
+ /*  2000 */   310,  310,  310,  310,  310,  310,  310,  310,  310,  310,
+ /*  2010 */   310,  310,  310,  310,  310,  310,  310,  310,  310,  310,
+ /*  2020 */   310,  310,  310,  310,  310,  310,  310,  310,  310,  310,
+ /*  2030 */   310,  310,  310,  310,  310,  310,  310,  310,  310,  310,
+ /*  2040 */   310,  310,  310,  310,  310,  310,  310,  310,  310,  310,
+ /*  2050 */   310,  310,  310,  310,  310,  310,  310,  310,  310,  310,
+ /*  2060 */   310,  310,  310,  310,  310,  310,  310,  310,  310,  310,
+ /*  2070 */   310,  310,  310,  310,  310,  310,  310,  310,  310,  310,
+ /*  2080 */   310,  310,  310,  310,  310,  310,  310,  310,  310,  310,
+ /*  2090 */   310,  310,  310,  310,  310,  310,  310,  310,  310,  310,
+ /*  2100 */   310,  310,  310,  310,  310,  310,  310,  310,  310,  310,
+ /*  2110 */   310,  310,  310,  310,  310,  310,  310,  310,  310,  310,
+ /*  2120 */   310,  310,  310,  310,  310,  310,  310,  310,  310,  310,
+ /*  2130 */   310,  310,  310,  310,  310,  310,  310,  310,  310,
 };
-#define YY_SHIFT_COUNT    (540)
+#define YY_SHIFT_COUNT    (550)
 #define YY_SHIFT_MIN      (0)
-#define YY_SHIFT_MAX      (1987)
+#define YY_SHIFT_MAX      (1948)
 static const unsigned short int yy_shift_ofst[] = {
- /*     0 */  1814, 1632, 1987, 1426, 1426,  128, 1482, 1633, 1703, 1877,
- /*    10 */  1877, 1877,   87,    0,    0,  264, 1106, 1877, 1877, 1877,
- /*    20 */  1877, 1877, 1877, 1877, 1877, 1877, 1877, 1877, 1877, 1877,
- /*    30 */   226,  226,  381,  381,  296,  193,  128,  128,  128,  128,
- /*    40 */   128,  128,   97,  194,  332,  429,  526,  623,  720,  817,
- /*    50 */   914,  934, 1086, 1238, 1106, 1106, 1106, 1106, 1106, 1106,
- /*    60 */  1106, 1106, 1106, 1106, 1106, 1106, 1106, 1106, 1106, 1106,
- /*    70 */  1106, 1106, 1258, 1106, 1355, 1375, 1375, 1817, 1877, 1877,
- /*    80 */  1877, 1877, 1877, 1877, 1877, 1877, 1877, 1877, 1877, 1877,
- /*    90 */  1877, 1877, 1877, 1877, 1877, 1877, 1877, 1877, 1877, 1877,
- /*   100 */  1877, 1877, 1877, 1877, 1877, 1877, 1877, 1877, 1877, 1877,
- /*   110 */  1937, 1877, 1877, 1877, 1877, 1877, 1877, 1877, 1877, 1877,
- /*   120 */  1877, 1877, 1877, 1877,   32,  129,  129,  129,  129,  129,
- /*   130 */    21,  152,  297,  494,  726,   65,  494,  514,  514,  494,
- /*   140 */   560,  560,  560,  560,  322,  599,   50, 2142, 2142,  155,
- /*   150 */   155,  155,  313,  392,  386,  392,  392,  481,  481,  200,
- /*   160 */   480,  684,  758,  494,  494,  494,  494,  494,  494,  494,
- /*   170 */   494,  494,  494,  494,  494,  494,  494,  494,  494,  494,
- /*   180 */   494,  494,  494,  494,  768,  768,  494,  166,  377,  377,
- /*   190 */   635,  835,  835,  635,  748,  987, 2142, 2142, 2142,  448,
- /*   200 */    45,   45,  403,  484,  502,  106,  525,  508,  528,  538,
- /*   210 */   494,  494,  494,  494,  494,  494,  494,  494,  494,   84,
- /*   220 */   494,  494,  494,  494,  494,  494,  494,  494,  494,  494,
- /*   230 */   494,  494,  267,  267,  267,  494,  494,  494,  494,  769,
- /*   240 */   494,  494,  494,    4,  477,  494,  494,  788,  494,  494,
- /*   250 */   494,  494,  494,  494,  494,  494,  727,    5,  135,  985,
- /*   260 */   985,  985,  985,  522,  135,  135,  797,  326,  875,  986,
- /*   270 */   968, 1036, 1036, 1038,  968,  968, 1038,  972, 1081, 1118,
- /*   280 */  1194, 1194, 1194, 1036,  757,  757,  946,  777, 1099, 1102,
- /*   290 */  1333, 1282, 1282, 1381, 1381, 1282, 1297, 1334, 1422, 1406,
- /*   300 */  1322, 1448, 1448, 1448, 1448, 1282, 1451, 1322, 1322, 1334,
- /*   310 */  1422, 1406, 1406, 1322, 1282, 1451, 1345, 1434, 1282, 1451,
- /*   320 */  1483, 1282, 1451, 1282, 1451, 1483, 1404, 1404, 1404, 1452,
- /*   330 */  1483, 1404, 1405, 1404, 1452, 1404, 1404, 1483, 1423, 1423,
- /*   340 */  1483, 1407, 1437, 1407, 1437, 1407, 1437, 1407, 1437, 1282,
- /*   350 */  1462, 1462, 1411, 1418, 1537, 1282, 1416, 1411, 1428, 1431,
- /*   360 */  1322, 1555, 1560, 1575, 1575, 1584, 1584, 1584, 2142, 2142,
- /*   370 */  2142, 2142, 2142, 2142, 2142, 2142, 2142, 2142, 2142, 2142,
- /*   380 */  2142, 2142, 2142, 2142,   61,  725,  374, 1080,  198,  771,
- /*   390 */   283, 1192, 1178, 1190, 1107, 1221, 1206, 1226, 1227, 1232,
- /*   400 */  1233, 1240, 1242, 1189, 1129, 1253,  216, 1210, 1247, 1248,
- /*   410 */  1249, 1131, 1151, 1274, 1293, 1220, 1214, 1605, 1608, 1591,
- /*   420 */  1459, 1600, 1525, 1606, 1598, 1602, 1494, 1492, 1513, 1614,
- /*   430 */  1504, 1620, 1510, 1625, 1647, 1515, 1507, 1531, 1595, 1621,
- /*   440 */  1514, 1607, 1610, 1612, 1613, 1536, 1552, 1634, 1533, 1669,
- /*   450 */  1666, 1651, 1566, 1522, 1609, 1650, 1611, 1603, 1639, 1547,
- /*   460 */  1574, 1659, 1664, 1667, 1561, 1569, 1668, 1622, 1671, 1672,
- /*   470 */  1665, 1673, 1624, 1670, 1674, 1630, 1662, 1677, 1559, 1681,
- /*   480 */  1682, 1549, 1684, 1685, 1683, 1688, 1690, 1692, 1691, 1695,
- /*   490 */  1694, 1585, 1698, 1705, 1617, 1696, 1707, 1596, 1709, 1697,
- /*   500 */  1702, 1704, 1711, 1652, 1675, 1658, 1708, 1676, 1656, 1714,
- /*   510 */  1726, 1731, 1730, 1729, 1733, 1722, 1734, 1709, 1737, 1738,
- /*   520 */  1742, 1743, 1741, 1745, 1747, 1759, 1749, 1750, 1752, 1753,
- /*   530 */  1751, 1755, 1757, 1660, 1653, 1655, 1657, 1661, 1761, 1777,
- /*   540 */  1792,
+ /*     0 */  1448, 1277, 1668, 1072, 1072,  340, 1122, 1225, 1332, 1481,
+ /*    10 */  1481, 1481,  335,    0,    0,  180,  897, 1481, 1481, 1481,
+ /*    20 */  1481, 1481, 1481, 1481, 1481, 1481, 1481, 1481, 1481, 1481,
+ /*    30 */   930,  930, 1020, 1020,  290,    1,  340,  340,  340,  340,
+ /*    40 */   340,  340,   40,  110,  219,  288,  327,  396,  435,  504,
+ /*    50 */   543,  612,  651,  720,  877,  897,  897,  897,  897,  897,
+ /*    60 */   897,  897,  897,  897,  897,  897,  897,  897,  897,  897,
+ /*    70 */   897,  897,  897,  917,  897, 1019,  763,  763, 1451, 1481,
+ /*    80 */  1481, 1481, 1481, 1481, 1481, 1481, 1481, 1481, 1481, 1481,
+ /*    90 */  1481, 1481, 1481, 1481, 1481, 1481, 1481, 1481, 1481, 1481,
+ /*   100 */  1481, 1481, 1481, 1481, 1481, 1481, 1481, 1481, 1481, 1481,
+ /*   110 */  1481, 1481, 1553, 1481, 1481, 1481, 1481, 1481, 1481, 1481,
+ /*   120 */  1481, 1481, 1481, 1481, 1481, 1481,  147,  258,  258,  258,
+ /*   130 */   258,  258,   79,   65,   84,  449,   19,  786,  449,  636,
+ /*   140 */   636,  449,  880,  880,  880,  880,  113,  142,  142,  472,
+ /*   150 */   150, 1958, 1958,  399,  399,  399,   93,  237,  341,  237,
+ /*   160 */   237, 1074, 1074,  437,  350,  704, 1080,  449,  449,  449,
+ /*   170 */   449,  449,  449,  449,  449,  449,  449,  449,  449,  449,
+ /*   180 */   449,  449,  449,  449,  449,  449,  449,  449,  818,  818,
+ /*   190 */   449, 1088,  217,  217,  734,  734, 1124, 1126, 1958, 1958,
+ /*   200 */  1958,  739,  840,  840,  453,  454,  511,  187,  563,  570,
+ /*   210 */   898,  669,  449,  449,  449,  449,  449,  449,  449,  449,
+ /*   220 */   449,  670,  449,  449,  449,  449,  449,  449,  449,  449,
+ /*   230 */   449,  449,  449,  449,  674,  674,  674,  449,  449,  449,
+ /*   240 */   449, 1034,  449,  449,  449,  972, 1107,  449,  449, 1193,
+ /*   250 */   449,  449,  449,  449,  449,  449,  449,  449,  260,  177,
+ /*   260 */   489, 1241, 1241, 1241, 1241, 1192,  489,  489,  952, 1197,
+ /*   270 */   625, 1235, 1139,  181,  181, 1086, 1139, 1139, 1086, 1187,
+ /*   280 */  1131, 1237, 1314, 1314, 1314,  181, 1245, 1245, 1109, 1299,
+ /*   290 */   549, 1340, 1606, 1531, 1531, 1639, 1639, 1531, 1538, 1576,
+ /*   300 */  1669, 1647, 1546, 1677, 1677, 1677, 1677, 1531, 1693, 1546,
+ /*   310 */  1546, 1576, 1669, 1647, 1647, 1546, 1531, 1693, 1567, 1661,
+ /*   320 */  1531, 1693, 1706, 1531, 1693, 1531, 1693, 1706, 1622, 1622,
+ /*   330 */  1622, 1675, 1720, 1720, 1706, 1622, 1619, 1622, 1675, 1622,
+ /*   340 */  1622, 1582, 1706, 1634, 1634, 1706, 1605, 1649, 1605, 1649,
+ /*   350 */  1605, 1649, 1605, 1649, 1531, 1679, 1679, 1690, 1690, 1629,
+ /*   360 */  1635, 1755, 1531, 1626, 1629, 1637, 1640, 1546, 1760, 1762,
+ /*   370 */  1773, 1773, 1787, 1787, 1787, 1958, 1958, 1958, 1958, 1958,
+ /*   380 */  1958, 1958, 1958, 1958, 1958, 1958, 1958, 1958, 1958, 1958,
+ /*   390 */   308,  835,  954, 1232,  879,  715,  728, 1323,  864, 1318,
+ /*   400 */  1253, 1373,  297, 1409, 1419, 1440, 1489, 1497, 1520, 1242,
+ /*   410 */  1309, 1447, 1435, 1341, 1521, 1525, 1392, 1548, 1329, 1354,
+ /*   420 */  1585, 1587, 1353, 1382, 1812, 1814, 1799, 1662, 1811, 1730,
+ /*   430 */  1813, 1807, 1808, 1695, 1685, 1707, 1810, 1696, 1817, 1697,
+ /*   440 */  1822, 1839, 1701, 1694, 1714, 1783, 1809, 1699, 1792, 1795,
+ /*   450 */  1796, 1797, 1724, 1740, 1821, 1717, 1856, 1853, 1837, 1747,
+ /*   460 */  1703, 1794, 1838, 1798, 1788, 1824, 1728, 1756, 1844, 1850,
+ /*   470 */  1852, 1743, 1750, 1854, 1815, 1855, 1857, 1851, 1858, 1816,
+ /*   480 */  1819, 1860, 1779, 1859, 1863, 1823, 1849, 1865, 1734, 1867,
+ /*   490 */  1868, 1869, 1870, 1871, 1872, 1874, 1875, 1877, 1876, 1878,
+ /*   500 */  1764, 1881, 1882, 1780, 1873, 1879, 1765, 1883, 1880, 1884,
+ /*   510 */  1885, 1886, 1820, 1835, 1825, 1887, 1840, 1828, 1888, 1889,
+ /*   520 */  1891, 1892, 1897, 1898, 1893, 1894, 1883, 1902, 1903, 1905,
+ /*   530 */  1906, 1904, 1909, 1911, 1923, 1913, 1914, 1915, 1916, 1918,
+ /*   540 */  1919, 1917, 1804, 1803, 1805, 1806, 1818, 1924, 1931, 1947,
+ /*   550 */  1948,
 };
-#define YY_REDUCE_COUNT (383)
-#define YY_REDUCE_MIN   (-257)
-#define YY_REDUCE_MAX   (1421)
+#define YY_REDUCE_COUNT (389)
+#define YY_REDUCE_MIN   (-262)
+#define YY_REDUCE_MAX   (1617)
 static const short yy_reduce_ofst[] = {
- /*     0 */  -168,  -17,  164,  214,  310, -166, -184,  -18,   98, -170,
- /*    10 */   305,  315, -163, -193, -178, -257,  395,  401,  476,  478,
- /*    20 */   512,  117,  527,  529,  503,  509,  532,  255,  552,  556,
- /*    30 */   558,  607,   37,  408,  594,  413,  462,  559,  561,  601,
- /*    40 */   610,  618, -254, -254, -254, -254, -254, -254, -254, -254,
- /*    50 */  -254, -254, -254, -254, -254, -254, -254, -254, -254, -254,
- /*    60 */  -254, -254, -254, -254, -254, -254, -254, -254, -254, -254,
- /*    70 */  -254, -254, -254, -254, -254, -254, -254, -111,  627,  650,
- /*    80 */   691,  697,  701,  703,  740,  742,  744,  767,  770,  790,
- /*    90 */   816,  818,  820,  822,  845,  857,  859,  861,  863,  865,
- /*   100 */   868,  870,  872,  874,  876,  888,  903,  906,  908,  915,
- /*   110 */   917,  922,  960,  962,  964,  988,  990,  992, 1015, 1017,
- /*   120 */  1029, 1033, 1035, 1040, -254, -254, -254, -254, -254, -254,
- /*   130 */  -254, -254, -254,  190,  270, -196,  160, -160,  450,  647,
- /*   140 */   260,  458,  260,  458,   78, -254, -254, -254, -254,  206,
- /*   150 */   206,  206,  320,  598,   -5,  675,  743, -148,  340, -125,
- /*   160 */   459,  466,  466,  693,  -93,  461,  479,  706,  710,  714,
- /*   170 */   716,  717,  169, -183,  325,  314,  704,  333,  747,  858,
- /*   180 */    -8,  819,  565,  755,  646,  660,  517,  265,  713,  791,
- /*   190 */   712,  795,  803,  918,  695,  860,  893,  935,  939, -181,
- /*   200 */  -172, -147,  -91,  -46,   -3,  162,  173,  231,  338,  437,
- /*   210 */   571,  614,  630,  651,  760,  931,  989, 1032, 1046, -218,
- /*   220 */    38, 1070, 1096, 1133, 1134, 1137, 1138, 1139, 1141, 1142,
- /*   230 */  1143, 1144,  292,  451, 1050, 1145, 1147, 1148, 1149,  813,
- /*   240 */  1161, 1162, 1163, 1108, 1049, 1166, 1168, 1146, 1169,  162,
- /*   250 */  1181, 1182, 1183, 1184, 1185, 1187, 1100, 1103, 1150, 1135,
- /*   260 */  1136, 1140, 1152,  813, 1150, 1150, 1153, 1173, 1195, 1090,
- /*   270 */  1154, 1167, 1170, 1104, 1155, 1156, 1109, 1172, 1174, 1179,
- /*   280 */  1177, 1188, 1205, 1171, 1160, 1196, 1121, 1165, 1203, 1228,
- /*   290 */  1157, 1244, 1245, 1158, 1159, 1250, 1175, 1193, 1191, 1241,
- /*   300 */  1231, 1243, 1257, 1259, 1261, 1276, 1280, 1254, 1255, 1230,
- /*   310 */  1234, 1269, 1270, 1260, 1292, 1303, 1223, 1225, 1309, 1313,
- /*   320 */  1296, 1317, 1319, 1320, 1323, 1300, 1307, 1308, 1310, 1304,
- /*   330 */  1311, 1315, 1314, 1318, 1316, 1321, 1325, 1324, 1271, 1272,
- /*   340 */  1332, 1294, 1298, 1299, 1301, 1302, 1306, 1312, 1329, 1356,
- /*   350 */  1256, 1263, 1327, 1326, 1305, 1369, 1330, 1340, 1343, 1346,
- /*   360 */  1350, 1391, 1394, 1410, 1412, 1415, 1417, 1419, 1328, 1331,
- /*   370 */  1335, 1402, 1398, 1400, 1401, 1403, 1408, 1396, 1397, 1409,
- /*   380 */  1414, 1420, 1421, 1413,
+ /*     0 */   490, -122,  545,  645,  650, -120, -189, -187, -184, -182,
+ /*    10 */  -178, -176,   45,   30,  200, -251, -134,  390,  392,  521,
+ /*    20 */   523,  213,  692,  821,  284,  589,  872,  666,  671,  866,
+ /*    30 */    71,  111,  273,  389,  686,  815,  904,  932,  948,  955,
+ /*    40 */   964,  969, -259, -259, -259, -259, -259, -259, -259, -259,
+ /*    50 */  -259, -259, -259, -259, -259, -259, -259, -259, -259, -259,
+ /*    60 */  -259, -259, -259, -259, -259, -259, -259, -259, -259, -259,
+ /*    70 */  -259, -259, -259, -259, -259, -259, -259, -259,  428,  430,
+ /*    80 */   899,  985, 1021, 1028, 1057, 1069, 1081, 1108, 1110, 1115,
+ /*    90 */  1117, 1123, 1149, 1154, 1159, 1170, 1174, 1178, 1183, 1194,
+ /*   100 */  1198, 1204, 1208, 1212, 1218, 1222, 1229, 1278, 1280, 1283,
+ /*   110 */  1285, 1313, 1316, 1320, 1322, 1325, 1327, 1330, 1366, 1371,
+ /*   120 */  1379, 1387, 1417, 1425, 1430, 1432, -259, -259, -259, -259,
+ /*   130 */  -259, -259, -259, -259, -259,  557,  974, -214, -174,   -9,
+ /*   140 */   431, -124,  806,  925,  806,  925,  251,  928,  940, -259,
+ /*   150 */  -259, -259, -259, -198, -198, -198,  127, -186, -168,  212,
+ /*   160 */   646,  617,  799, -262,  555,  220,  220,  491,  605, 1040,
+ /*   170 */  1060,  699,  -11,  600,  848,  862,  345, -129,  724,  -91,
+ /*   180 */   158,  749,  716,  900,  304,  822,  929,  926,  499,  793,
+ /*   190 */   322,  892,  813,  845,  958, 1056,  751,  905, 1133, 1062,
+ /*   200 */   803, -210, -185, -179, -148, -167,  -89,  121,  274,  281,
+ /*   210 */   320,  336,  439,  663,  711,  957, 1064, 1068, 1116, 1127,
+ /*   220 */  1134, -196, 1147, 1180, 1184, 1195, 1203, 1209, 1254, 1263,
+ /*   230 */  1275, 1288, 1304, 1310,  205,  422,  638, 1319, 1324, 1346,
+ /*   240 */  1360, 1168, 1364, 1370, 1372,  869, 1189, 1380, 1399, 1276,
+ /*   250 */  1403,  121, 1405, 1420, 1426, 1427, 1428, 1429, 1249, 1282,
+ /*   260 */  1344, 1375, 1376, 1377, 1388, 1168, 1344, 1344, 1384, 1411,
+ /*   270 */  1436, 1349, 1389, 1386, 1391, 1361, 1407, 1408, 1365, 1431,
+ /*   280 */  1433, 1434, 1439, 1441, 1442, 1396, 1416, 1418, 1390, 1421,
+ /*   290 */  1437, 1472, 1381, 1478, 1480, 1397, 1400, 1487, 1412, 1444,
+ /*   300 */  1438, 1463, 1453, 1464, 1465, 1467, 1469, 1514, 1517, 1473,
+ /*   310 */  1474, 1452, 1449, 1490, 1491, 1475, 1522, 1526, 1443, 1445,
+ /*   320 */  1528, 1530, 1513, 1534, 1536, 1537, 1539, 1516, 1523, 1524,
+ /*   330 */  1527, 1519, 1529, 1532, 1540, 1541, 1535, 1542, 1544, 1545,
+ /*   340 */  1547, 1450, 1543, 1477, 1482, 1551, 1505, 1508, 1512, 1509,
+ /*   350 */  1515, 1518, 1533, 1552, 1573, 1466, 1468, 1549, 1550, 1555,
+ /*   360 */  1554, 1510, 1583, 1511, 1556, 1559, 1561, 1565, 1588, 1592,
+ /*   370 */  1601, 1602, 1607, 1608, 1609, 1498, 1557, 1558, 1610, 1600,
+ /*   380 */  1603, 1611, 1612, 1613, 1596, 1597, 1614, 1615, 1617, 1616,
 };
 static const YYACTIONTYPE yy_default[] = {
- /*     0 */  1536, 1536, 1536, 1376, 1159, 1265, 1159, 1159, 1159, 1376,
- /*    10 */  1376, 1376, 1159, 1295, 1295, 1429, 1190, 1159, 1159, 1159,
- /*    20 */  1159, 1159, 1159, 1159, 1159, 1159, 1159, 1375, 1159, 1159,
- /*    30 */  1159, 1159, 1459, 1459, 1159, 1159, 1159, 1159, 1159, 1159,
- /*    40 */  1159, 1159, 1159, 1301, 1159, 1159, 1159, 1159, 1159, 1377,
- /*    50 */  1378, 1159, 1159, 1159, 1428, 1430, 1393, 1311, 1310, 1309,
- /*    60 */  1308, 1411, 1282, 1306, 1299, 1303, 1371, 1372, 1370, 1374,
- /*    70 */  1378, 1377, 1159, 1302, 1342, 1356, 1341, 1159, 1159, 1159,
- /*    80 */  1159, 1159, 1159, 1159, 1159, 1159, 1159, 1159, 1159, 1159,
- /*    90 */  1159, 1159, 1159, 1159, 1159, 1159, 1159, 1159, 1159, 1159,
- /*   100 */  1159, 1159, 1159, 1159, 1159, 1159, 1159, 1159, 1159, 1159,
- /*   110 */  1159, 1159, 1159, 1159, 1159, 1159, 1159, 1159, 1159, 1159,
- /*   120 */  1159, 1159, 1159, 1159, 1350, 1355, 1361, 1354, 1351, 1344,
- /*   130 */  1343, 1345, 1346, 1159, 1180, 1229, 1159, 1159, 1159, 1159,
- /*   140 */  1447, 1446, 1159, 1159, 1190, 1347, 1348, 1358, 1357, 1436,
- /*   150 */  1492, 1491, 1394, 1159, 1159, 1159, 1159, 1159, 1159, 1459,
- /*   160 */  1159, 1159, 1159, 1159, 1159, 1159, 1159, 1159, 1159, 1159,
- /*   170 */  1159, 1159, 1159, 1159, 1159, 1159, 1159, 1159, 1159, 1159,
- /*   180 */  1159, 1159, 1159, 1159, 1459, 1459, 1159, 1190, 1459, 1459,
- /*   190 */  1186, 1336, 1335, 1186, 1289, 1159, 1442, 1265, 1256, 1159,
- /*   200 */  1159, 1159, 1159, 1159, 1159, 1159, 1159, 1159, 1159, 1159,
- /*   210 */  1159, 1159, 1159, 1433, 1431, 1159, 1159, 1159, 1159, 1159,
- /*   220 */  1159, 1159, 1159, 1159, 1159, 1159, 1159, 1159, 1159, 1159,
- /*   230 */  1159, 1159, 1159, 1159, 1159, 1159, 1159, 1159, 1159, 1159,
- /*   240 */  1159, 1159, 1159, 1261, 1159, 1159, 1159, 1159, 1159, 1159,
- /*   250 */  1159, 1159, 1159, 1159, 1159, 1486, 1159, 1406, 1243, 1261,
- /*   260 */  1261, 1261, 1261, 1263, 1244, 1242, 1255, 1190, 1166, 1528,
- /*   270 */  1305, 1284, 1284, 1525, 1305, 1305, 1525, 1204, 1506, 1201,
- /*   280 */  1295, 1295, 1295, 1284, 1289, 1289, 1373, 1262, 1255, 1159,
- /*   290 */  1528, 1270, 1270, 1527, 1527, 1270, 1394, 1314, 1320, 1232,
- /*   300 */  1305, 1238, 1238, 1238, 1238, 1270, 1177, 1305, 1305, 1314,
- /*   310 */  1320, 1232, 1232, 1305, 1270, 1177, 1410, 1522, 1270, 1177,
- /*   320 */  1384, 1270, 1177, 1270, 1177, 1384, 1230, 1230, 1230, 1219,
- /*   330 */  1384, 1230, 1204, 1230, 1219, 1230, 1230, 1384, 1388, 1388,
- /*   340 */  1384, 1288, 1283, 1288, 1283, 1288, 1283, 1288, 1283, 1270,
- /*   350 */  1469, 1469, 1300, 1289, 1379, 1270, 1159, 1300, 1298, 1296,
- /*   360 */  1305, 1183, 1222, 1489, 1489, 1485, 1485, 1485, 1533, 1533,
- /*   370 */  1442, 1501, 1190, 1190, 1190, 1190, 1501, 1206, 1206, 1190,
- /*   380 */  1190, 1190, 1190, 1501, 1159, 1159, 1159, 1159, 1159, 1159,
- /*   390 */  1496, 1159, 1395, 1274, 1159, 1159, 1159, 1159, 1159, 1159,
- /*   400 */  1159, 1159, 1159, 1159, 1159, 1159, 1159, 1159, 1159, 1159,
- /*   410 */  1159, 1159, 1159, 1159, 1159, 1159, 1325, 1159, 1162, 1439,
- /*   420 */  1159, 1159, 1437, 1159, 1159, 1159, 1159, 1159, 1159, 1275,
- /*   430 */  1159, 1159, 1159, 1159, 1159, 1159, 1159, 1159, 1159, 1159,
- /*   440 */  1159, 1159, 1159, 1159, 1159, 1159, 1159, 1159, 1524, 1159,
- /*   450 */  1159, 1159, 1159, 1159, 1159, 1409, 1408, 1159, 1159, 1272,
- /*   460 */  1159, 1159, 1159, 1159, 1159, 1159, 1159, 1159, 1159, 1159,
- /*   470 */  1159, 1159, 1159, 1159, 1159, 1159, 1159, 1159, 1159, 1159,
- /*   480 */  1159, 1159, 1159, 1159, 1159, 1159, 1159, 1159, 1159, 1159,
- /*   490 */  1159, 1159, 1159, 1159, 1159, 1159, 1159, 1159, 1297, 1159,
- /*   500 */  1159, 1159, 1159, 1159, 1159, 1159, 1159, 1159, 1159, 1159,
- /*   510 */  1159, 1159, 1159, 1474, 1290, 1159, 1159, 1515, 1159, 1159,
- /*   520 */  1159, 1159, 1159, 1159, 1159, 1159, 1159, 1159, 1159, 1159,
- /*   530 */  1159, 1159, 1510, 1246, 1327, 1159, 1326, 1330, 1159, 1171,
- /*   540 */  1159,
+ /*     0 */  1573, 1573, 1573, 1409, 1186, 1295, 1186, 1186, 1186, 1409,
+ /*    10 */  1409, 1409, 1186, 1325, 1325, 1462, 1217, 1186, 1186, 1186,
+ /*    20 */  1186, 1186, 1186, 1186, 1186, 1186, 1186, 1408, 1186, 1186,
+ /*    30 */  1186, 1186, 1492, 1492, 1186, 1186, 1186, 1186, 1186, 1186,
+ /*    40 */  1186, 1186, 1186, 1334, 1186, 1186, 1186, 1186, 1186, 1186,
+ /*    50 */  1410, 1411, 1186, 1186, 1186, 1461, 1463, 1426, 1344, 1343,
+ /*    60 */  1342, 1341, 1444, 1312, 1339, 1332, 1336, 1404, 1405, 1403,
+ /*    70 */  1407, 1411, 1410, 1186, 1335, 1375, 1389, 1374, 1186, 1186,
+ /*    80 */  1186, 1186, 1186, 1186, 1186, 1186, 1186, 1186, 1186, 1186,
+ /*    90 */  1186, 1186, 1186, 1186, 1186, 1186, 1186, 1186, 1186, 1186,
+ /*   100 */  1186, 1186, 1186, 1186, 1186, 1186, 1186, 1186, 1186, 1186,
+ /*   110 */  1186, 1186, 1186, 1186, 1186, 1186, 1186, 1186, 1186, 1186,
+ /*   120 */  1186, 1186, 1186, 1186, 1186, 1186, 1383, 1388, 1394, 1387,
+ /*   130 */  1384, 1377, 1376, 1378, 1379, 1186, 1207, 1259, 1186, 1186,
+ /*   140 */  1186, 1186, 1480, 1479, 1186, 1186, 1217, 1369, 1368, 1380,
+ /*   150 */  1381, 1391, 1390, 1469, 1527, 1526, 1427, 1186, 1186, 1186,
+ /*   160 */  1186, 1186, 1186, 1492, 1186, 1186, 1186, 1186, 1186, 1186,
+ /*   170 */  1186, 1186, 1186, 1186, 1186, 1186, 1186, 1186, 1186, 1186,
+ /*   180 */  1186, 1186, 1186, 1186, 1186, 1186, 1186, 1186, 1492, 1492,
+ /*   190 */  1186, 1217, 1492, 1492, 1213, 1213, 1319, 1186, 1475, 1295,
+ /*   200 */  1286, 1186, 1186, 1186, 1186, 1186, 1186, 1186, 1186, 1186,
+ /*   210 */  1186, 1186, 1186, 1186, 1186, 1466, 1464, 1186, 1186, 1186,
+ /*   220 */  1186, 1186, 1186, 1186, 1186, 1186, 1186, 1186, 1186, 1186,
+ /*   230 */  1186, 1186, 1186, 1186, 1186, 1186, 1186, 1186, 1186, 1186,
+ /*   240 */  1186, 1186, 1186, 1186, 1186, 1291, 1186, 1186, 1186, 1186,
+ /*   250 */  1186, 1186, 1186, 1186, 1186, 1186, 1186, 1521, 1186, 1439,
+ /*   260 */  1273, 1291, 1291, 1291, 1291, 1293, 1274, 1272, 1285, 1218,
+ /*   270 */  1193, 1565, 1338, 1314, 1314, 1562, 1338, 1338, 1562, 1234,
+ /*   280 */  1543, 1229, 1325, 1325, 1325, 1314, 1319, 1319, 1406, 1292,
+ /*   290 */  1285, 1186, 1565, 1300, 1300, 1564, 1564, 1300, 1427, 1347,
+ /*   300 */  1353, 1262, 1338, 1268, 1268, 1268, 1268, 1300, 1204, 1338,
+ /*   310 */  1338, 1347, 1353, 1262, 1262, 1338, 1300, 1204, 1443, 1559,
+ /*   320 */  1300, 1204, 1417, 1300, 1204, 1300, 1204, 1417, 1260, 1260,
+ /*   330 */  1260, 1249, 1186, 1186, 1417, 1260, 1234, 1260, 1249, 1260,
+ /*   340 */  1260, 1510, 1417, 1421, 1421, 1417, 1318, 1313, 1318, 1313,
+ /*   350 */  1318, 1313, 1318, 1313, 1300, 1502, 1502, 1328, 1328, 1333,
+ /*   360 */  1319, 1412, 1300, 1186, 1333, 1331, 1329, 1338, 1210, 1252,
+ /*   370 */  1524, 1524, 1520, 1520, 1520, 1570, 1570, 1475, 1536, 1217,
+ /*   380 */  1217, 1217, 1217, 1536, 1236, 1236, 1218, 1218, 1217, 1536,
+ /*   390 */  1186, 1186, 1186, 1186, 1186, 1186, 1531, 1186, 1428, 1304,
+ /*   400 */  1186, 1186, 1186, 1186, 1186, 1186, 1186, 1186, 1186, 1186,
+ /*   410 */  1186, 1186, 1186, 1186, 1186, 1186, 1186, 1186, 1186, 1186,
+ /*   420 */  1186, 1186, 1186, 1358, 1186, 1189, 1472, 1186, 1186, 1470,
+ /*   430 */  1186, 1186, 1186, 1186, 1186, 1186, 1305, 1186, 1186, 1186,
+ /*   440 */  1186, 1186, 1186, 1186, 1186, 1186, 1186, 1186, 1186, 1186,
+ /*   450 */  1186, 1186, 1186, 1186, 1186, 1561, 1186, 1186, 1186, 1186,
+ /*   460 */  1186, 1186, 1442, 1441, 1186, 1186, 1302, 1186, 1186, 1186,
+ /*   470 */  1186, 1186, 1186, 1186, 1186, 1186, 1186, 1186, 1186, 1186,
+ /*   480 */  1232, 1186, 1186, 1186, 1186, 1186, 1186, 1186, 1186, 1186,
+ /*   490 */  1186, 1186, 1186, 1186, 1186, 1186, 1186, 1186, 1186, 1186,
+ /*   500 */  1186, 1186, 1186, 1186, 1186, 1186, 1186, 1330, 1186, 1186,
+ /*   510 */  1186, 1186, 1186, 1186, 1186, 1186, 1186, 1186, 1186, 1186,
+ /*   520 */  1186, 1186, 1507, 1320, 1186, 1186, 1552, 1186, 1186, 1186,
+ /*   530 */  1186, 1186, 1186, 1186, 1186, 1186, 1186, 1186, 1186, 1186,
+ /*   540 */  1186, 1547, 1276, 1360, 1186, 1359, 1363, 1186, 1198, 1186,
+ /*   550 */  1186,
 };
 /********** End of lemon-generated parsing tables *****************************/
 
@@ -149177,6 +153346,9 @@ static const YYCODETYPE yyFallback[] = {
    59,  /*       VIEW => ID */
    59,  /*    VIRTUAL => ID */
    59,  /*       WITH => ID */
+   59,  /*      NULLS => ID */
+   59,  /*      FIRST => ID */
+   59,  /*       LAST => ID */
    59,  /*    CURRENT => ID */
    59,  /*  FOLLOWING => ID */
    59,  /*  PARTITION => ID */
@@ -149187,9 +153359,92 @@ static const YYCODETYPE yyFallback[] = {
    59,  /*     GROUPS => ID */
    59,  /*     OTHERS => ID */
    59,  /*       TIES => ID */
+   59,  /*  GENERATED => ID */
+   59,  /*     ALWAYS => ID */
    59,  /*    REINDEX => ID */
    59,  /*     RENAME => ID */
    59,  /*   CTIME_KW => ID */
+    0,  /*        ANY => nothing */
+    0,  /*     BITAND => nothing */
+    0,  /*      BITOR => nothing */
+    0,  /*     LSHIFT => nothing */
+    0,  /*     RSHIFT => nothing */
+    0,  /*       PLUS => nothing */
+    0,  /*      MINUS => nothing */
+    0,  /*       STAR => nothing */
+    0,  /*      SLASH => nothing */
+    0,  /*        REM => nothing */
+    0,  /*     CONCAT => nothing */
+    0,  /*    COLLATE => nothing */
+    0,  /*     BITNOT => nothing */
+    0,  /*         ON => nothing */
+    0,  /*    INDEXED => nothing */
+    0,  /*     STRING => nothing */
+    0,  /*    JOIN_KW => nothing */
+    0,  /* CONSTRAINT => nothing */
+    0,  /*    DEFAULT => nothing */
+    0,  /*       NULL => nothing */
+    0,  /*    PRIMARY => nothing */
+    0,  /*     UNIQUE => nothing */
+    0,  /*      CHECK => nothing */
+    0,  /* REFERENCES => nothing */
+    0,  /*   AUTOINCR => nothing */
+    0,  /*     INSERT => nothing */
+    0,  /*     DELETE => nothing */
+    0,  /*     UPDATE => nothing */
+    0,  /*        SET => nothing */
+    0,  /* DEFERRABLE => nothing */
+    0,  /*    FOREIGN => nothing */
+    0,  /*       DROP => nothing */
+    0,  /*      UNION => nothing */
+    0,  /*        ALL => nothing */
+    0,  /*     EXCEPT => nothing */
+    0,  /*  INTERSECT => nothing */
+    0,  /*     SELECT => nothing */
+    0,  /*     VALUES => nothing */
+    0,  /*   DISTINCT => nothing */
+    0,  /*        DOT => nothing */
+    0,  /*       FROM => nothing */
+    0,  /*       JOIN => nothing */
+    0,  /*      USING => nothing */
+    0,  /*      ORDER => nothing */
+    0,  /*      GROUP => nothing */
+    0,  /*     HAVING => nothing */
+    0,  /*      LIMIT => nothing */
+    0,  /*      WHERE => nothing */
+    0,  /*       INTO => nothing */
+    0,  /*    NOTHING => nothing */
+    0,  /*      FLOAT => nothing */
+    0,  /*       BLOB => nothing */
+    0,  /*    INTEGER => nothing */
+    0,  /*   VARIABLE => nothing */
+    0,  /*       CASE => nothing */
+    0,  /*       WHEN => nothing */
+    0,  /*       THEN => nothing */
+    0,  /*       ELSE => nothing */
+    0,  /*      INDEX => nothing */
+    0,  /*      ALTER => nothing */
+    0,  /*        ADD => nothing */
+    0,  /*     WINDOW => nothing */
+    0,  /*       OVER => nothing */
+    0,  /*     FILTER => nothing */
+    0,  /*     COLUMN => nothing */
+    0,  /* AGG_FUNCTION => nothing */
+    0,  /* AGG_COLUMN => nothing */
+    0,  /*  TRUEFALSE => nothing */
+    0,  /*      ISNOT => nothing */
+    0,  /*   FUNCTION => nothing */
+    0,  /*     UMINUS => nothing */
+    0,  /*      UPLUS => nothing */
+    0,  /*      TRUTH => nothing */
+    0,  /*   REGISTER => nothing */
+    0,  /*     VECTOR => nothing */
+    0,  /* SELECT_COLUMN => nothing */
+    0,  /* IF_NULL_ROW => nothing */
+    0,  /*   ASTERISK => nothing */
+    0,  /*       SPAN => nothing */
+    0,  /*      SPACE => nothing */
+    0,  /*    ILLEGAL => nothing */
 };
 #endif /* YYFALLBACK */
 
@@ -149359,225 +153614,234 @@ static const char *const yyTokenName[] = {
   /*   79 */ "VIEW",
   /*   80 */ "VIRTUAL",
   /*   81 */ "WITH",
-  /*   82 */ "CURRENT",
-  /*   83 */ "FOLLOWING",
-  /*   84 */ "PARTITION",
-  /*   85 */ "PRECEDING",
-  /*   86 */ "RANGE",
-  /*   87 */ "UNBOUNDED",
-  /*   88 */ "EXCLUDE",
-  /*   89 */ "GROUPS",
-  /*   90 */ "OTHERS",
-  /*   91 */ "TIES",
-  /*   92 */ "REINDEX",
-  /*   93 */ "RENAME",
-  /*   94 */ "CTIME_KW",
-  /*   95 */ "ANY",
-  /*   96 */ "BITAND",
-  /*   97 */ "BITOR",
-  /*   98 */ "LSHIFT",
-  /*   99 */ "RSHIFT",
-  /*  100 */ "PLUS",
-  /*  101 */ "MINUS",
-  /*  102 */ "STAR",
-  /*  103 */ "SLASH",
-  /*  104 */ "REM",
-  /*  105 */ "CONCAT",
-  /*  106 */ "COLLATE",
-  /*  107 */ "BITNOT",
-  /*  108 */ "ON",
-  /*  109 */ "INDEXED",
-  /*  110 */ "STRING",
-  /*  111 */ "JOIN_KW",
-  /*  112 */ "CONSTRAINT",
-  /*  113 */ "DEFAULT",
-  /*  114 */ "NULL",
-  /*  115 */ "PRIMARY",
-  /*  116 */ "UNIQUE",
-  /*  117 */ "CHECK",
-  /*  118 */ "REFERENCES",
-  /*  119 */ "AUTOINCR",
-  /*  120 */ "INSERT",
-  /*  121 */ "DELETE",
-  /*  122 */ "UPDATE",
-  /*  123 */ "SET",
-  /*  124 */ "DEFERRABLE",
-  /*  125 */ "FOREIGN",
-  /*  126 */ "DROP",
-  /*  127 */ "UNION",
-  /*  128 */ "ALL",
-  /*  129 */ "EXCEPT",
-  /*  130 */ "INTERSECT",
-  /*  131 */ "SELECT",
-  /*  132 */ "VALUES",
-  /*  133 */ "DISTINCT",
-  /*  134 */ "DOT",
-  /*  135 */ "FROM",
-  /*  136 */ "JOIN",
-  /*  137 */ "USING",
-  /*  138 */ "ORDER",
-  /*  139 */ "GROUP",
-  /*  140 */ "HAVING",
-  /*  141 */ "LIMIT",
-  /*  142 */ "WHERE",
-  /*  143 */ "INTO",
-  /*  144 */ "NOTHING",
-  /*  145 */ "FLOAT",
-  /*  146 */ "BLOB",
-  /*  147 */ "INTEGER",
-  /*  148 */ "VARIABLE",
-  /*  149 */ "CASE",
-  /*  150 */ "WHEN",
-  /*  151 */ "THEN",
-  /*  152 */ "ELSE",
-  /*  153 */ "INDEX",
-  /*  154 */ "ALTER",
-  /*  155 */ "ADD",
-  /*  156 */ "WINDOW",
-  /*  157 */ "OVER",
-  /*  158 */ "FILTER",
-  /*  159 */ "TRUEFALSE",
-  /*  160 */ "ISNOT",
-  /*  161 */ "FUNCTION",
-  /*  162 */ "COLUMN",
-  /*  163 */ "AGG_FUNCTION",
-  /*  164 */ "AGG_COLUMN",
-  /*  165 */ "UMINUS",
-  /*  166 */ "UPLUS",
-  /*  167 */ "TRUTH",
-  /*  168 */ "REGISTER",
-  /*  169 */ "VECTOR",
-  /*  170 */ "SELECT_COLUMN",
-  /*  171 */ "IF_NULL_ROW",
-  /*  172 */ "ASTERISK",
-  /*  173 */ "SPAN",
-  /*  174 */ "SPACE",
-  /*  175 */ "ILLEGAL",
-  /*  176 */ "input",
-  /*  177 */ "cmdlist",
-  /*  178 */ "ecmd",
-  /*  179 */ "cmdx",
-  /*  180 */ "explain",
-  /*  181 */ "cmd",
-  /*  182 */ "transtype",
-  /*  183 */ "trans_opt",
-  /*  184 */ "nm",
-  /*  185 */ "savepoint_opt",
-  /*  186 */ "create_table",
-  /*  187 */ "create_table_args",
-  /*  188 */ "createkw",
-  /*  189 */ "temp",
-  /*  190 */ "ifnotexists",
-  /*  191 */ "dbnm",
-  /*  192 */ "columnlist",
-  /*  193 */ "conslist_opt",
-  /*  194 */ "table_options",
-  /*  195 */ "select",
-  /*  196 */ "columnname",
-  /*  197 */ "carglist",
-  /*  198 */ "typetoken",
-  /*  199 */ "typename",
-  /*  200 */ "signed",
-  /*  201 */ "plus_num",
-  /*  202 */ "minus_num",
-  /*  203 */ "scanpt",
-  /*  204 */ "ccons",
-  /*  205 */ "term",
-  /*  206 */ "expr",
-  /*  207 */ "onconf",
-  /*  208 */ "sortorder",
-  /*  209 */ "autoinc",
-  /*  210 */ "eidlist_opt",
-  /*  211 */ "refargs",
-  /*  212 */ "defer_subclause",
-  /*  213 */ "refarg",
-  /*  214 */ "refact",
-  /*  215 */ "init_deferred_pred_opt",
-  /*  216 */ "conslist",
-  /*  217 */ "tconscomma",
-  /*  218 */ "tcons",
-  /*  219 */ "sortlist",
-  /*  220 */ "eidlist",
-  /*  221 */ "defer_subclause_opt",
-  /*  222 */ "orconf",
-  /*  223 */ "resolvetype",
-  /*  224 */ "raisetype",
-  /*  225 */ "ifexists",
-  /*  226 */ "fullname",
-  /*  227 */ "selectnowith",
-  /*  228 */ "oneselect",
-  /*  229 */ "wqlist",
-  /*  230 */ "multiselect_op",
-  /*  231 */ "distinct",
-  /*  232 */ "selcollist",
-  /*  233 */ "from",
-  /*  234 */ "where_opt",
-  /*  235 */ "groupby_opt",
-  /*  236 */ "having_opt",
-  /*  237 */ "orderby_opt",
-  /*  238 */ "limit_opt",
-  /*  239 */ "window_clause",
-  /*  240 */ "values",
-  /*  241 */ "nexprlist",
-  /*  242 */ "sclp",
-  /*  243 */ "as",
-  /*  244 */ "seltablist",
-  /*  245 */ "stl_prefix",
-  /*  246 */ "joinop",
-  /*  247 */ "indexed_opt",
-  /*  248 */ "on_opt",
-  /*  249 */ "using_opt",
-  /*  250 */ "exprlist",
-  /*  251 */ "xfullname",
-  /*  252 */ "idlist",
-  /*  253 */ "with",
-  /*  254 */ "setlist",
-  /*  255 */ "insert_cmd",
-  /*  256 */ "idlist_opt",
-  /*  257 */ "upsert",
-  /*  258 */ "over_clause",
-  /*  259 */ "likeop",
-  /*  260 */ "between_op",
-  /*  261 */ "in_op",
-  /*  262 */ "paren_exprlist",
-  /*  263 */ "case_operand",
-  /*  264 */ "case_exprlist",
-  /*  265 */ "case_else",
-  /*  266 */ "uniqueflag",
-  /*  267 */ "collate",
-  /*  268 */ "vinto",
-  /*  269 */ "nmnum",
-  /*  270 */ "trigger_decl",
-  /*  271 */ "trigger_cmd_list",
-  /*  272 */ "trigger_time",
-  /*  273 */ "trigger_event",
-  /*  274 */ "foreach_clause",
-  /*  275 */ "when_clause",
-  /*  276 */ "trigger_cmd",
-  /*  277 */ "trnm",
-  /*  278 */ "tridxby",
-  /*  279 */ "database_kw_opt",
-  /*  280 */ "key_opt",
-  /*  281 */ "add_column_fullname",
-  /*  282 */ "kwcolumn_opt",
-  /*  283 */ "create_vtab",
-  /*  284 */ "vtabarglist",
-  /*  285 */ "vtabarg",
-  /*  286 */ "vtabargtoken",
-  /*  287 */ "lp",
-  /*  288 */ "anylist",
-  /*  289 */ "windowdefn_list",
-  /*  290 */ "windowdefn",
-  /*  291 */ "window",
-  /*  292 */ "frame_opt",
-  /*  293 */ "part_opt",
-  /*  294 */ "filter_opt",
-  /*  295 */ "range_or_rows",
-  /*  296 */ "frame_bound",
-  /*  297 */ "frame_bound_s",
-  /*  298 */ "frame_bound_e",
-  /*  299 */ "frame_exclude_opt",
-  /*  300 */ "frame_exclude",
+  /*   82 */ "NULLS",
+  /*   83 */ "FIRST",
+  /*   84 */ "LAST",
+  /*   85 */ "CURRENT",
+  /*   86 */ "FOLLOWING",
+  /*   87 */ "PARTITION",
+  /*   88 */ "PRECEDING",
+  /*   89 */ "RANGE",
+  /*   90 */ "UNBOUNDED",
+  /*   91 */ "EXCLUDE",
+  /*   92 */ "GROUPS",
+  /*   93 */ "OTHERS",
+  /*   94 */ "TIES",
+  /*   95 */ "GENERATED",
+  /*   96 */ "ALWAYS",
+  /*   97 */ "REINDEX",
+  /*   98 */ "RENAME",
+  /*   99 */ "CTIME_KW",
+  /*  100 */ "ANY",
+  /*  101 */ "BITAND",
+  /*  102 */ "BITOR",
+  /*  103 */ "LSHIFT",
+  /*  104 */ "RSHIFT",
+  /*  105 */ "PLUS",
+  /*  106 */ "MINUS",
+  /*  107 */ "STAR",
+  /*  108 */ "SLASH",
+  /*  109 */ "REM",
+  /*  110 */ "CONCAT",
+  /*  111 */ "COLLATE",
+  /*  112 */ "BITNOT",
+  /*  113 */ "ON",
+  /*  114 */ "INDEXED",
+  /*  115 */ "STRING",
+  /*  116 */ "JOIN_KW",
+  /*  117 */ "CONSTRAINT",
+  /*  118 */ "DEFAULT",
+  /*  119 */ "NULL",
+  /*  120 */ "PRIMARY",
+  /*  121 */ "UNIQUE",
+  /*  122 */ "CHECK",
+  /*  123 */ "REFERENCES",
+  /*  124 */ "AUTOINCR",
+  /*  125 */ "INSERT",
+  /*  126 */ "DELETE",
+  /*  127 */ "UPDATE",
+  /*  128 */ "SET",
+  /*  129 */ "DEFERRABLE",
+  /*  130 */ "FOREIGN",
+  /*  131 */ "DROP",
+  /*  132 */ "UNION",
+  /*  133 */ "ALL",
+  /*  134 */ "EXCEPT",
+  /*  135 */ "INTERSECT",
+  /*  136 */ "SELECT",
+  /*  137 */ "VALUES",
+  /*  138 */ "DISTINCT",
+  /*  139 */ "DOT",
+  /*  140 */ "FROM",
+  /*  141 */ "JOIN",
+  /*  142 */ "USING",
+  /*  143 */ "ORDER",
+  /*  144 */ "GROUP",
+  /*  145 */ "HAVING",
+  /*  146 */ "LIMIT",
+  /*  147 */ "WHERE",
+  /*  148 */ "INTO",
+  /*  149 */ "NOTHING",
+  /*  150 */ "FLOAT",
+  /*  151 */ "BLOB",
+  /*  152 */ "INTEGER",
+  /*  153 */ "VARIABLE",
+  /*  154 */ "CASE",
+  /*  155 */ "WHEN",
+  /*  156 */ "THEN",
+  /*  157 */ "ELSE",
+  /*  158 */ "INDEX",
+  /*  159 */ "ALTER",
+  /*  160 */ "ADD",
+  /*  161 */ "WINDOW",
+  /*  162 */ "OVER",
+  /*  163 */ "FILTER",
+  /*  164 */ "COLUMN",
+  /*  165 */ "AGG_FUNCTION",
+  /*  166 */ "AGG_COLUMN",
+  /*  167 */ "TRUEFALSE",
+  /*  168 */ "ISNOT",
+  /*  169 */ "FUNCTION",
+  /*  170 */ "UMINUS",
+  /*  171 */ "UPLUS",
+  /*  172 */ "TRUTH",
+  /*  173 */ "REGISTER",
+  /*  174 */ "VECTOR",
+  /*  175 */ "SELECT_COLUMN",
+  /*  176 */ "IF_NULL_ROW",
+  /*  177 */ "ASTERISK",
+  /*  178 */ "SPAN",
+  /*  179 */ "SPACE",
+  /*  180 */ "ILLEGAL",
+  /*  181 */ "input",
+  /*  182 */ "cmdlist",
+  /*  183 */ "ecmd",
+  /*  184 */ "cmdx",
+  /*  185 */ "explain",
+  /*  186 */ "cmd",
+  /*  187 */ "transtype",
+  /*  188 */ "trans_opt",
+  /*  189 */ "nm",
+  /*  190 */ "savepoint_opt",
+  /*  191 */ "create_table",
+  /*  192 */ "create_table_args",
+  /*  193 */ "createkw",
+  /*  194 */ "temp",
+  /*  195 */ "ifnotexists",
+  /*  196 */ "dbnm",
+  /*  197 */ "columnlist",
+  /*  198 */ "conslist_opt",
+  /*  199 */ "table_options",
+  /*  200 */ "select",
+  /*  201 */ "columnname",
+  /*  202 */ "carglist",
+  /*  203 */ "typetoken",
+  /*  204 */ "typename",
+  /*  205 */ "signed",
+  /*  206 */ "plus_num",
+  /*  207 */ "minus_num",
+  /*  208 */ "scanpt",
+  /*  209 */ "scantok",
+  /*  210 */ "ccons",
+  /*  211 */ "term",
+  /*  212 */ "expr",
+  /*  213 */ "onconf",
+  /*  214 */ "sortorder",
+  /*  215 */ "autoinc",
+  /*  216 */ "eidlist_opt",
+  /*  217 */ "refargs",
+  /*  218 */ "defer_subclause",
+  /*  219 */ "generated",
+  /*  220 */ "refarg",
+  /*  221 */ "refact",
+  /*  222 */ "init_deferred_pred_opt",
+  /*  223 */ "conslist",
+  /*  224 */ "tconscomma",
+  /*  225 */ "tcons",
+  /*  226 */ "sortlist",
+  /*  227 */ "eidlist",
+  /*  228 */ "defer_subclause_opt",
+  /*  229 */ "orconf",
+  /*  230 */ "resolvetype",
+  /*  231 */ "raisetype",
+  /*  232 */ "ifexists",
+  /*  233 */ "fullname",
+  /*  234 */ "selectnowith",
+  /*  235 */ "oneselect",
+  /*  236 */ "wqlist",
+  /*  237 */ "multiselect_op",
+  /*  238 */ "distinct",
+  /*  239 */ "selcollist",
+  /*  240 */ "from",
+  /*  241 */ "where_opt",
+  /*  242 */ "groupby_opt",
+  /*  243 */ "having_opt",
+  /*  244 */ "orderby_opt",
+  /*  245 */ "limit_opt",
+  /*  246 */ "window_clause",
+  /*  247 */ "values",
+  /*  248 */ "nexprlist",
+  /*  249 */ "sclp",
+  /*  250 */ "as",
+  /*  251 */ "seltablist",
+  /*  252 */ "stl_prefix",
+  /*  253 */ "joinop",
+  /*  254 */ "indexed_opt",
+  /*  255 */ "on_opt",
+  /*  256 */ "using_opt",
+  /*  257 */ "exprlist",
+  /*  258 */ "xfullname",
+  /*  259 */ "idlist",
+  /*  260 */ "nulls",
+  /*  261 */ "with",
+  /*  262 */ "setlist",
+  /*  263 */ "insert_cmd",
+  /*  264 */ "idlist_opt",
+  /*  265 */ "upsert",
+  /*  266 */ "filter_over",
+  /*  267 */ "likeop",
+  /*  268 */ "between_op",
+  /*  269 */ "in_op",
+  /*  270 */ "paren_exprlist",
+  /*  271 */ "case_operand",
+  /*  272 */ "case_exprlist",
+  /*  273 */ "case_else",
+  /*  274 */ "uniqueflag",
+  /*  275 */ "collate",
+  /*  276 */ "vinto",
+  /*  277 */ "nmnum",
+  /*  278 */ "trigger_decl",
+  /*  279 */ "trigger_cmd_list",
+  /*  280 */ "trigger_time",
+  /*  281 */ "trigger_event",
+  /*  282 */ "foreach_clause",
+  /*  283 */ "when_clause",
+  /*  284 */ "trigger_cmd",
+  /*  285 */ "trnm",
+  /*  286 */ "tridxby",
+  /*  287 */ "database_kw_opt",
+  /*  288 */ "key_opt",
+  /*  289 */ "add_column_fullname",
+  /*  290 */ "kwcolumn_opt",
+  /*  291 */ "create_vtab",
+  /*  292 */ "vtabarglist",
+  /*  293 */ "vtabarg",
+  /*  294 */ "vtabargtoken",
+  /*  295 */ "lp",
+  /*  296 */ "anylist",
+  /*  297 */ "windowdefn_list",
+  /*  298 */ "windowdefn",
+  /*  299 */ "window",
+  /*  300 */ "frame_opt",
+  /*  301 */ "part_opt",
+  /*  302 */ "filter_clause",
+  /*  303 */ "over_clause",
+  /*  304 */ "range_or_rows",
+  /*  305 */ "frame_bound",
+  /*  306 */ "frame_bound_s",
+  /*  307 */ "frame_bound_e",
+  /*  308 */ "frame_exclude_opt",
+  /*  309 */ "frame_exclude",
 };
 #endif /* defined(YYCOVERAGE) || !defined(NDEBUG) */
 
@@ -149614,352 +153878,362 @@ static const char *const yyRuleName[] = {
  /*  26 */ "typetoken ::= typename LP signed COMMA signed RP",
  /*  27 */ "typename ::= typename ID|STRING",
  /*  28 */ "scanpt ::=",
- /*  29 */ "ccons ::= CONSTRAINT nm",
- /*  30 */ "ccons ::= DEFAULT scanpt term scanpt",
- /*  31 */ "ccons ::= DEFAULT LP expr RP",
- /*  32 */ "ccons ::= DEFAULT PLUS term scanpt",
- /*  33 */ "ccons ::= DEFAULT MINUS term scanpt",
- /*  34 */ "ccons ::= DEFAULT scanpt ID|INDEXED",
- /*  35 */ "ccons ::= NOT NULL onconf",
- /*  36 */ "ccons ::= PRIMARY KEY sortorder onconf autoinc",
- /*  37 */ "ccons ::= UNIQUE onconf",
- /*  38 */ "ccons ::= CHECK LP expr RP",
- /*  39 */ "ccons ::= REFERENCES nm eidlist_opt refargs",
- /*  40 */ "ccons ::= defer_subclause",
- /*  41 */ "ccons ::= COLLATE ID|STRING",
- /*  42 */ "autoinc ::=",
- /*  43 */ "autoinc ::= AUTOINCR",
- /*  44 */ "refargs ::=",
- /*  45 */ "refargs ::= refargs refarg",
- /*  46 */ "refarg ::= MATCH nm",
- /*  47 */ "refarg ::= ON INSERT refact",
- /*  48 */ "refarg ::= ON DELETE refact",
- /*  49 */ "refarg ::= ON UPDATE refact",
- /*  50 */ "refact ::= SET NULL",
- /*  51 */ "refact ::= SET DEFAULT",
- /*  52 */ "refact ::= CASCADE",
- /*  53 */ "refact ::= RESTRICT",
- /*  54 */ "refact ::= NO ACTION",
- /*  55 */ "defer_subclause ::= NOT DEFERRABLE init_deferred_pred_opt",
- /*  56 */ "defer_subclause ::= DEFERRABLE init_deferred_pred_opt",
- /*  57 */ "init_deferred_pred_opt ::=",
- /*  58 */ "init_deferred_pred_opt ::= INITIALLY DEFERRED",
- /*  59 */ "init_deferred_pred_opt ::= INITIALLY IMMEDIATE",
- /*  60 */ "conslist_opt ::=",
- /*  61 */ "tconscomma ::= COMMA",
- /*  62 */ "tcons ::= CONSTRAINT nm",
- /*  63 */ "tcons ::= PRIMARY KEY LP sortlist autoinc RP onconf",
- /*  64 */ "tcons ::= UNIQUE LP sortlist RP onconf",
- /*  65 */ "tcons ::= CHECK LP expr RP onconf",
- /*  66 */ "tcons ::= FOREIGN KEY LP eidlist RP REFERENCES nm eidlist_opt refargs defer_subclause_opt",
- /*  67 */ "defer_subclause_opt ::=",
- /*  68 */ "onconf ::=",
- /*  69 */ "onconf ::= ON CONFLICT resolvetype",
- /*  70 */ "orconf ::=",
- /*  71 */ "orconf ::= OR resolvetype",
- /*  72 */ "resolvetype ::= IGNORE",
- /*  73 */ "resolvetype ::= REPLACE",
- /*  74 */ "cmd ::= DROP TABLE ifexists fullname",
- /*  75 */ "ifexists ::= IF EXISTS",
- /*  76 */ "ifexists ::=",
- /*  77 */ "cmd ::= createkw temp VIEW ifnotexists nm dbnm eidlist_opt AS select",
- /*  78 */ "cmd ::= DROP VIEW ifexists fullname",
- /*  79 */ "cmd ::= select",
- /*  80 */ "select ::= WITH wqlist selectnowith",
- /*  81 */ "select ::= WITH RECURSIVE wqlist selectnowith",
- /*  82 */ "select ::= selectnowith",
- /*  83 */ "selectnowith ::= selectnowith multiselect_op oneselect",
- /*  84 */ "multiselect_op ::= UNION",
- /*  85 */ "multiselect_op ::= UNION ALL",
- /*  86 */ "multiselect_op ::= EXCEPT|INTERSECT",
- /*  87 */ "oneselect ::= SELECT distinct selcollist from where_opt groupby_opt having_opt orderby_opt limit_opt",
- /*  88 */ "oneselect ::= SELECT distinct selcollist from where_opt groupby_opt having_opt window_clause orderby_opt limit_opt",
- /*  89 */ "values ::= VALUES LP nexprlist RP",
- /*  90 */ "values ::= values COMMA LP nexprlist RP",
- /*  91 */ "distinct ::= DISTINCT",
- /*  92 */ "distinct ::= ALL",
- /*  93 */ "distinct ::=",
- /*  94 */ "sclp ::=",
- /*  95 */ "selcollist ::= sclp scanpt expr scanpt as",
- /*  96 */ "selcollist ::= sclp scanpt STAR",
- /*  97 */ "selcollist ::= sclp scanpt nm DOT STAR",
- /*  98 */ "as ::= AS nm",
- /*  99 */ "as ::=",
- /* 100 */ "from ::=",
- /* 101 */ "from ::= FROM seltablist",
- /* 102 */ "stl_prefix ::= seltablist joinop",
- /* 103 */ "stl_prefix ::=",
- /* 104 */ "seltablist ::= stl_prefix nm dbnm as indexed_opt on_opt using_opt",
- /* 105 */ "seltablist ::= stl_prefix nm dbnm LP exprlist RP as on_opt using_opt",
- /* 106 */ "seltablist ::= stl_prefix LP select RP as on_opt using_opt",
- /* 107 */ "seltablist ::= stl_prefix LP seltablist RP as on_opt using_opt",
- /* 108 */ "dbnm ::=",
- /* 109 */ "dbnm ::= DOT nm",
- /* 110 */ "fullname ::= nm",
- /* 111 */ "fullname ::= nm DOT nm",
- /* 112 */ "xfullname ::= nm",
- /* 113 */ "xfullname ::= nm DOT nm",
- /* 114 */ "xfullname ::= nm DOT nm AS nm",
- /* 115 */ "xfullname ::= nm AS nm",
- /* 116 */ "joinop ::= COMMA|JOIN",
- /* 117 */ "joinop ::= JOIN_KW JOIN",
- /* 118 */ "joinop ::= JOIN_KW nm JOIN",
- /* 119 */ "joinop ::= JOIN_KW nm nm JOIN",
- /* 120 */ "on_opt ::= ON expr",
- /* 121 */ "on_opt ::=",
- /* 122 */ "indexed_opt ::=",
- /* 123 */ "indexed_opt ::= INDEXED BY nm",
- /* 124 */ "indexed_opt ::= NOT INDEXED",
- /* 125 */ "using_opt ::= USING LP idlist RP",
- /* 126 */ "using_opt ::=",
- /* 127 */ "orderby_opt ::=",
- /* 128 */ "orderby_opt ::= ORDER BY sortlist",
- /* 129 */ "sortlist ::= sortlist COMMA expr sortorder",
- /* 130 */ "sortlist ::= expr sortorder",
- /* 131 */ "sortorder ::= ASC",
- /* 132 */ "sortorder ::= DESC",
- /* 133 */ "sortorder ::=",
- /* 134 */ "groupby_opt ::=",
- /* 135 */ "groupby_opt ::= GROUP BY nexprlist",
- /* 136 */ "having_opt ::=",
- /* 137 */ "having_opt ::= HAVING expr",
- /* 138 */ "limit_opt ::=",
- /* 139 */ "limit_opt ::= LIMIT expr",
- /* 140 */ "limit_opt ::= LIMIT expr OFFSET expr",
- /* 141 */ "limit_opt ::= LIMIT expr COMMA expr",
- /* 142 */ "cmd ::= with DELETE FROM xfullname indexed_opt where_opt",
- /* 143 */ "where_opt ::=",
- /* 144 */ "where_opt ::= WHERE expr",
- /* 145 */ "cmd ::= with UPDATE orconf xfullname indexed_opt SET setlist where_opt",
- /* 146 */ "setlist ::= setlist COMMA nm EQ expr",
- /* 147 */ "setlist ::= setlist COMMA LP idlist RP EQ expr",
- /* 148 */ "setlist ::= nm EQ expr",
- /* 149 */ "setlist ::= LP idlist RP EQ expr",
- /* 150 */ "cmd ::= with insert_cmd INTO xfullname idlist_opt select upsert",
- /* 151 */ "cmd ::= with insert_cmd INTO xfullname idlist_opt DEFAULT VALUES",
- /* 152 */ "upsert ::=",
- /* 153 */ "upsert ::= ON CONFLICT LP sortlist RP where_opt DO UPDATE SET setlist where_opt",
- /* 154 */ "upsert ::= ON CONFLICT LP sortlist RP where_opt DO NOTHING",
- /* 155 */ "upsert ::= ON CONFLICT DO NOTHING",
- /* 156 */ "insert_cmd ::= INSERT orconf",
- /* 157 */ "insert_cmd ::= REPLACE",
- /* 158 */ "idlist_opt ::=",
- /* 159 */ "idlist_opt ::= LP idlist RP",
- /* 160 */ "idlist ::= idlist COMMA nm",
- /* 161 */ "idlist ::= nm",
- /* 162 */ "expr ::= LP expr RP",
- /* 163 */ "expr ::= ID|INDEXED",
- /* 164 */ "expr ::= JOIN_KW",
- /* 165 */ "expr ::= nm DOT nm",
- /* 166 */ "expr ::= nm DOT nm DOT nm",
- /* 167 */ "term ::= NULL|FLOAT|BLOB",
- /* 168 */ "term ::= STRING",
- /* 169 */ "term ::= INTEGER",
- /* 170 */ "expr ::= VARIABLE",
- /* 171 */ "expr ::= expr COLLATE ID|STRING",
- /* 172 */ "expr ::= CAST LP expr AS typetoken RP",
- /* 173 */ "expr ::= ID|INDEXED LP distinct exprlist RP",
- /* 174 */ "expr ::= ID|INDEXED LP STAR RP",
- /* 175 */ "expr ::= ID|INDEXED LP distinct exprlist RP over_clause",
- /* 176 */ "expr ::= ID|INDEXED LP STAR RP over_clause",
- /* 177 */ "term ::= CTIME_KW",
- /* 178 */ "expr ::= LP nexprlist COMMA expr RP",
- /* 179 */ "expr ::= expr AND expr",
- /* 180 */ "expr ::= expr OR expr",
- /* 181 */ "expr ::= expr LT|GT|GE|LE expr",
- /* 182 */ "expr ::= expr EQ|NE expr",
- /* 183 */ "expr ::= expr BITAND|BITOR|LSHIFT|RSHIFT expr",
- /* 184 */ "expr ::= expr PLUS|MINUS expr",
- /* 185 */ "expr ::= expr STAR|SLASH|REM expr",
- /* 186 */ "expr ::= expr CONCAT expr",
- /* 187 */ "likeop ::= NOT LIKE_KW|MATCH",
- /* 188 */ "expr ::= expr likeop expr",
- /* 189 */ "expr ::= expr likeop expr ESCAPE expr",
- /* 190 */ "expr ::= expr ISNULL|NOTNULL",
- /* 191 */ "expr ::= expr NOT NULL",
- /* 192 */ "expr ::= expr IS expr",
- /* 193 */ "expr ::= expr IS NOT expr",
- /* 194 */ "expr ::= NOT expr",
- /* 195 */ "expr ::= BITNOT expr",
- /* 196 */ "expr ::= PLUS|MINUS expr",
- /* 197 */ "between_op ::= BETWEEN",
- /* 198 */ "between_op ::= NOT BETWEEN",
- /* 199 */ "expr ::= expr between_op expr AND expr",
- /* 200 */ "in_op ::= IN",
- /* 201 */ "in_op ::= NOT IN",
- /* 202 */ "expr ::= expr in_op LP exprlist RP",
- /* 203 */ "expr ::= LP select RP",
- /* 204 */ "expr ::= expr in_op LP select RP",
- /* 205 */ "expr ::= expr in_op nm dbnm paren_exprlist",
- /* 206 */ "expr ::= EXISTS LP select RP",
- /* 207 */ "expr ::= CASE case_operand case_exprlist case_else END",
- /* 208 */ "case_exprlist ::= case_exprlist WHEN expr THEN expr",
- /* 209 */ "case_exprlist ::= WHEN expr THEN expr",
- /* 210 */ "case_else ::= ELSE expr",
- /* 211 */ "case_else ::=",
- /* 212 */ "case_operand ::= expr",
- /* 213 */ "case_operand ::=",
- /* 214 */ "exprlist ::=",
- /* 215 */ "nexprlist ::= nexprlist COMMA expr",
- /* 216 */ "nexprlist ::= expr",
- /* 217 */ "paren_exprlist ::=",
- /* 218 */ "paren_exprlist ::= LP exprlist RP",
- /* 219 */ "cmd ::= createkw uniqueflag INDEX ifnotexists nm dbnm ON nm LP sortlist RP where_opt",
- /* 220 */ "uniqueflag ::= UNIQUE",
- /* 221 */ "uniqueflag ::=",
- /* 222 */ "eidlist_opt ::=",
- /* 223 */ "eidlist_opt ::= LP eidlist RP",
- /* 224 */ "eidlist ::= eidlist COMMA nm collate sortorder",
- /* 225 */ "eidlist ::= nm collate sortorder",
- /* 226 */ "collate ::=",
- /* 227 */ "collate ::= COLLATE ID|STRING",
- /* 228 */ "cmd ::= DROP INDEX ifexists fullname",
- /* 229 */ "cmd ::= VACUUM vinto",
- /* 230 */ "cmd ::= VACUUM nm vinto",
- /* 231 */ "vinto ::= INTO expr",
- /* 232 */ "vinto ::=",
- /* 233 */ "cmd ::= PRAGMA nm dbnm",
- /* 234 */ "cmd ::= PRAGMA nm dbnm EQ nmnum",
- /* 235 */ "cmd ::= PRAGMA nm dbnm LP nmnum RP",
- /* 236 */ "cmd ::= PRAGMA nm dbnm EQ minus_num",
- /* 237 */ "cmd ::= PRAGMA nm dbnm LP minus_num RP",
- /* 238 */ "plus_num ::= PLUS INTEGER|FLOAT",
- /* 239 */ "minus_num ::= MINUS INTEGER|FLOAT",
- /* 240 */ "cmd ::= createkw trigger_decl BEGIN trigger_cmd_list END",
- /* 241 */ "trigger_decl ::= temp TRIGGER ifnotexists nm dbnm trigger_time trigger_event ON fullname foreach_clause when_clause",
- /* 242 */ "trigger_time ::= BEFORE|AFTER",
- /* 243 */ "trigger_time ::= INSTEAD OF",
- /* 244 */ "trigger_time ::=",
- /* 245 */ "trigger_event ::= DELETE|INSERT",
- /* 246 */ "trigger_event ::= UPDATE",
- /* 247 */ "trigger_event ::= UPDATE OF idlist",
- /* 248 */ "when_clause ::=",
- /* 249 */ "when_clause ::= WHEN expr",
- /* 250 */ "trigger_cmd_list ::= trigger_cmd_list trigger_cmd SEMI",
- /* 251 */ "trigger_cmd_list ::= trigger_cmd SEMI",
- /* 252 */ "trnm ::= nm DOT nm",
- /* 253 */ "tridxby ::= INDEXED BY nm",
- /* 254 */ "tridxby ::= NOT INDEXED",
- /* 255 */ "trigger_cmd ::= UPDATE orconf trnm tridxby SET setlist where_opt scanpt",
- /* 256 */ "trigger_cmd ::= scanpt insert_cmd INTO trnm idlist_opt select upsert scanpt",
- /* 257 */ "trigger_cmd ::= DELETE FROM trnm tridxby where_opt scanpt",
- /* 258 */ "trigger_cmd ::= scanpt select scanpt",
- /* 259 */ "expr ::= RAISE LP IGNORE RP",
- /* 260 */ "expr ::= RAISE LP raisetype COMMA nm RP",
- /* 261 */ "raisetype ::= ROLLBACK",
- /* 262 */ "raisetype ::= ABORT",
- /* 263 */ "raisetype ::= FAIL",
- /* 264 */ "cmd ::= DROP TRIGGER ifexists fullname",
- /* 265 */ "cmd ::= ATTACH database_kw_opt expr AS expr key_opt",
- /* 266 */ "cmd ::= DETACH database_kw_opt expr",
- /* 267 */ "key_opt ::=",
- /* 268 */ "key_opt ::= KEY expr",
- /* 269 */ "cmd ::= REINDEX",
- /* 270 */ "cmd ::= REINDEX nm dbnm",
- /* 271 */ "cmd ::= ANALYZE",
- /* 272 */ "cmd ::= ANALYZE nm dbnm",
- /* 273 */ "cmd ::= ALTER TABLE fullname RENAME TO nm",
- /* 274 */ "cmd ::= ALTER TABLE add_column_fullname ADD kwcolumn_opt columnname carglist",
- /* 275 */ "add_column_fullname ::= fullname",
- /* 276 */ "cmd ::= ALTER TABLE fullname RENAME kwcolumn_opt nm TO nm",
- /* 277 */ "cmd ::= create_vtab",
- /* 278 */ "cmd ::= create_vtab LP vtabarglist RP",
- /* 279 */ "create_vtab ::= createkw VIRTUAL TABLE ifnotexists nm dbnm USING nm",
- /* 280 */ "vtabarg ::=",
- /* 281 */ "vtabargtoken ::= ANY",
- /* 282 */ "vtabargtoken ::= lp anylist RP",
- /* 283 */ "lp ::= LP",
- /* 284 */ "with ::= WITH wqlist",
- /* 285 */ "with ::= WITH RECURSIVE wqlist",
- /* 286 */ "wqlist ::= nm eidlist_opt AS LP select RP",
- /* 287 */ "wqlist ::= wqlist COMMA nm eidlist_opt AS LP select RP",
- /* 288 */ "windowdefn_list ::= windowdefn",
- /* 289 */ "windowdefn_list ::= windowdefn_list COMMA windowdefn",
- /* 290 */ "windowdefn ::= nm AS LP window RP",
- /* 291 */ "window ::= PARTITION BY nexprlist orderby_opt frame_opt",
- /* 292 */ "window ::= nm PARTITION BY nexprlist orderby_opt frame_opt",
- /* 293 */ "window ::= ORDER BY sortlist frame_opt",
- /* 294 */ "window ::= nm ORDER BY sortlist frame_opt",
- /* 295 */ "window ::= frame_opt",
- /* 296 */ "window ::= nm frame_opt",
- /* 297 */ "frame_opt ::=",
- /* 298 */ "frame_opt ::= range_or_rows frame_bound_s frame_exclude_opt",
- /* 299 */ "frame_opt ::= range_or_rows BETWEEN frame_bound_s AND frame_bound_e frame_exclude_opt",
- /* 300 */ "range_or_rows ::= RANGE|ROWS|GROUPS",
- /* 301 */ "frame_bound_s ::= frame_bound",
- /* 302 */ "frame_bound_s ::= UNBOUNDED PRECEDING",
- /* 303 */ "frame_bound_e ::= frame_bound",
- /* 304 */ "frame_bound_e ::= UNBOUNDED FOLLOWING",
- /* 305 */ "frame_bound ::= expr PRECEDING|FOLLOWING",
- /* 306 */ "frame_bound ::= CURRENT ROW",
- /* 307 */ "frame_exclude_opt ::=",
- /* 308 */ "frame_exclude_opt ::= EXCLUDE frame_exclude",
- /* 309 */ "frame_exclude ::= NO OTHERS",
- /* 310 */ "frame_exclude ::= CURRENT ROW",
- /* 311 */ "frame_exclude ::= GROUP|TIES",
- /* 312 */ "window_clause ::= WINDOW windowdefn_list",
- /* 313 */ "over_clause ::= filter_opt OVER LP window RP",
- /* 314 */ "over_clause ::= filter_opt OVER nm",
- /* 315 */ "filter_opt ::=",
- /* 316 */ "filter_opt ::= FILTER LP WHERE expr RP",
- /* 317 */ "input ::= cmdlist",
- /* 318 */ "cmdlist ::= cmdlist ecmd",
- /* 319 */ "cmdlist ::= ecmd",
- /* 320 */ "ecmd ::= SEMI",
- /* 321 */ "ecmd ::= cmdx SEMI",
- /* 322 */ "ecmd ::= explain cmdx",
- /* 323 */ "trans_opt ::=",
- /* 324 */ "trans_opt ::= TRANSACTION",
- /* 325 */ "trans_opt ::= TRANSACTION nm",
- /* 326 */ "savepoint_opt ::= SAVEPOINT",
- /* 327 */ "savepoint_opt ::=",
- /* 328 */ "cmd ::= create_table create_table_args",
- /* 329 */ "columnlist ::= columnlist COMMA columnname carglist",
- /* 330 */ "columnlist ::= columnname carglist",
- /* 331 */ "nm ::= ID|INDEXED",
- /* 332 */ "nm ::= STRING",
- /* 333 */ "nm ::= JOIN_KW",
- /* 334 */ "typetoken ::= typename",
- /* 335 */ "typename ::= ID|STRING",
- /* 336 */ "signed ::= plus_num",
- /* 337 */ "signed ::= minus_num",
- /* 338 */ "carglist ::= carglist ccons",
- /* 339 */ "carglist ::=",
- /* 340 */ "ccons ::= NULL onconf",
- /* 341 */ "conslist_opt ::= COMMA conslist",
- /* 342 */ "conslist ::= conslist tconscomma tcons",
- /* 343 */ "conslist ::= tcons",
- /* 344 */ "tconscomma ::=",
- /* 345 */ "defer_subclause_opt ::= defer_subclause",
- /* 346 */ "resolvetype ::= raisetype",
- /* 347 */ "selectnowith ::= oneselect",
- /* 348 */ "oneselect ::= values",
- /* 349 */ "sclp ::= selcollist COMMA",
- /* 350 */ "as ::= ID|STRING",
- /* 351 */ "expr ::= term",
- /* 352 */ "likeop ::= LIKE_KW|MATCH",
- /* 353 */ "exprlist ::= nexprlist",
- /* 354 */ "nmnum ::= plus_num",
- /* 355 */ "nmnum ::= nm",
- /* 356 */ "nmnum ::= ON",
- /* 357 */ "nmnum ::= DELETE",
- /* 358 */ "nmnum ::= DEFAULT",
- /* 359 */ "plus_num ::= INTEGER|FLOAT",
- /* 360 */ "foreach_clause ::=",
- /* 361 */ "foreach_clause ::= FOR EACH ROW",
- /* 362 */ "trnm ::= nm",
- /* 363 */ "tridxby ::=",
- /* 364 */ "database_kw_opt ::= DATABASE",
- /* 365 */ "database_kw_opt ::=",
- /* 366 */ "kwcolumn_opt ::=",
- /* 367 */ "kwcolumn_opt ::= COLUMNKW",
- /* 368 */ "vtabarglist ::= vtabarg",
- /* 369 */ "vtabarglist ::= vtabarglist COMMA vtabarg",
- /* 370 */ "vtabarg ::= vtabarg vtabargtoken",
- /* 371 */ "anylist ::=",
- /* 372 */ "anylist ::= anylist LP anylist RP",
- /* 373 */ "anylist ::= anylist ANY",
- /* 374 */ "with ::=",
+ /*  29 */ "scantok ::=",
+ /*  30 */ "ccons ::= CONSTRAINT nm",
+ /*  31 */ "ccons ::= DEFAULT scantok term",
+ /*  32 */ "ccons ::= DEFAULT LP expr RP",
+ /*  33 */ "ccons ::= DEFAULT PLUS scantok term",
+ /*  34 */ "ccons ::= DEFAULT MINUS scantok term",
+ /*  35 */ "ccons ::= DEFAULT scantok ID|INDEXED",
+ /*  36 */ "ccons ::= NOT NULL onconf",
+ /*  37 */ "ccons ::= PRIMARY KEY sortorder onconf autoinc",
+ /*  38 */ "ccons ::= UNIQUE onconf",
+ /*  39 */ "ccons ::= CHECK LP expr RP",
+ /*  40 */ "ccons ::= REFERENCES nm eidlist_opt refargs",
+ /*  41 */ "ccons ::= defer_subclause",
+ /*  42 */ "ccons ::= COLLATE ID|STRING",
+ /*  43 */ "generated ::= LP expr RP",
+ /*  44 */ "generated ::= LP expr RP ID",
+ /*  45 */ "autoinc ::=",
+ /*  46 */ "autoinc ::= AUTOINCR",
+ /*  47 */ "refargs ::=",
+ /*  48 */ "refargs ::= refargs refarg",
+ /*  49 */ "refarg ::= MATCH nm",
+ /*  50 */ "refarg ::= ON INSERT refact",
+ /*  51 */ "refarg ::= ON DELETE refact",
+ /*  52 */ "refarg ::= ON UPDATE refact",
+ /*  53 */ "refact ::= SET NULL",
+ /*  54 */ "refact ::= SET DEFAULT",
+ /*  55 */ "refact ::= CASCADE",
+ /*  56 */ "refact ::= RESTRICT",
+ /*  57 */ "refact ::= NO ACTION",
+ /*  58 */ "defer_subclause ::= NOT DEFERRABLE init_deferred_pred_opt",
+ /*  59 */ "defer_subclause ::= DEFERRABLE init_deferred_pred_opt",
+ /*  60 */ "init_deferred_pred_opt ::=",
+ /*  61 */ "init_deferred_pred_opt ::= INITIALLY DEFERRED",
+ /*  62 */ "init_deferred_pred_opt ::= INITIALLY IMMEDIATE",
+ /*  63 */ "conslist_opt ::=",
+ /*  64 */ "tconscomma ::= COMMA",
+ /*  65 */ "tcons ::= CONSTRAINT nm",
+ /*  66 */ "tcons ::= PRIMARY KEY LP sortlist autoinc RP onconf",
+ /*  67 */ "tcons ::= UNIQUE LP sortlist RP onconf",
+ /*  68 */ "tcons ::= CHECK LP expr RP onconf",
+ /*  69 */ "tcons ::= FOREIGN KEY LP eidlist RP REFERENCES nm eidlist_opt refargs defer_subclause_opt",
+ /*  70 */ "defer_subclause_opt ::=",
+ /*  71 */ "onconf ::=",
+ /*  72 */ "onconf ::= ON CONFLICT resolvetype",
+ /*  73 */ "orconf ::=",
+ /*  74 */ "orconf ::= OR resolvetype",
+ /*  75 */ "resolvetype ::= IGNORE",
+ /*  76 */ "resolvetype ::= REPLACE",
+ /*  77 */ "cmd ::= DROP TABLE ifexists fullname",
+ /*  78 */ "ifexists ::= IF EXISTS",
+ /*  79 */ "ifexists ::=",
+ /*  80 */ "cmd ::= createkw temp VIEW ifnotexists nm dbnm eidlist_opt AS select",
+ /*  81 */ "cmd ::= DROP VIEW ifexists fullname",
+ /*  82 */ "cmd ::= select",
+ /*  83 */ "select ::= WITH wqlist selectnowith",
+ /*  84 */ "select ::= WITH RECURSIVE wqlist selectnowith",
+ /*  85 */ "select ::= selectnowith",
+ /*  86 */ "selectnowith ::= selectnowith multiselect_op oneselect",
+ /*  87 */ "multiselect_op ::= UNION",
+ /*  88 */ "multiselect_op ::= UNION ALL",
+ /*  89 */ "multiselect_op ::= EXCEPT|INTERSECT",
+ /*  90 */ "oneselect ::= SELECT distinct selcollist from where_opt groupby_opt having_opt orderby_opt limit_opt",
+ /*  91 */ "oneselect ::= SELECT distinct selcollist from where_opt groupby_opt having_opt window_clause orderby_opt limit_opt",
+ /*  92 */ "values ::= VALUES LP nexprlist RP",
+ /*  93 */ "values ::= values COMMA LP nexprlist RP",
+ /*  94 */ "distinct ::= DISTINCT",
+ /*  95 */ "distinct ::= ALL",
+ /*  96 */ "distinct ::=",
+ /*  97 */ "sclp ::=",
+ /*  98 */ "selcollist ::= sclp scanpt expr scanpt as",
+ /*  99 */ "selcollist ::= sclp scanpt STAR",
+ /* 100 */ "selcollist ::= sclp scanpt nm DOT STAR",
+ /* 101 */ "as ::= AS nm",
+ /* 102 */ "as ::=",
+ /* 103 */ "from ::=",
+ /* 104 */ "from ::= FROM seltablist",
+ /* 105 */ "stl_prefix ::= seltablist joinop",
+ /* 106 */ "stl_prefix ::=",
+ /* 107 */ "seltablist ::= stl_prefix nm dbnm as indexed_opt on_opt using_opt",
+ /* 108 */ "seltablist ::= stl_prefix nm dbnm LP exprlist RP as on_opt using_opt",
+ /* 109 */ "seltablist ::= stl_prefix LP select RP as on_opt using_opt",
+ /* 110 */ "seltablist ::= stl_prefix LP seltablist RP as on_opt using_opt",
+ /* 111 */ "dbnm ::=",
+ /* 112 */ "dbnm ::= DOT nm",
+ /* 113 */ "fullname ::= nm",
+ /* 114 */ "fullname ::= nm DOT nm",
+ /* 115 */ "xfullname ::= nm",
+ /* 116 */ "xfullname ::= nm DOT nm",
+ /* 117 */ "xfullname ::= nm DOT nm AS nm",
+ /* 118 */ "xfullname ::= nm AS nm",
+ /* 119 */ "joinop ::= COMMA|JOIN",
+ /* 120 */ "joinop ::= JOIN_KW JOIN",
+ /* 121 */ "joinop ::= JOIN_KW nm JOIN",
+ /* 122 */ "joinop ::= JOIN_KW nm nm JOIN",
+ /* 123 */ "on_opt ::= ON expr",
+ /* 124 */ "on_opt ::=",
+ /* 125 */ "indexed_opt ::=",
+ /* 126 */ "indexed_opt ::= INDEXED BY nm",
+ /* 127 */ "indexed_opt ::= NOT INDEXED",
+ /* 128 */ "using_opt ::= USING LP idlist RP",
+ /* 129 */ "using_opt ::=",
+ /* 130 */ "orderby_opt ::=",
+ /* 131 */ "orderby_opt ::= ORDER BY sortlist",
+ /* 132 */ "sortlist ::= sortlist COMMA expr sortorder nulls",
+ /* 133 */ "sortlist ::= expr sortorder nulls",
+ /* 134 */ "sortorder ::= ASC",
+ /* 135 */ "sortorder ::= DESC",
+ /* 136 */ "sortorder ::=",
+ /* 137 */ "nulls ::= NULLS FIRST",
+ /* 138 */ "nulls ::= NULLS LAST",
+ /* 139 */ "nulls ::=",
+ /* 140 */ "groupby_opt ::=",
+ /* 141 */ "groupby_opt ::= GROUP BY nexprlist",
+ /* 142 */ "having_opt ::=",
+ /* 143 */ "having_opt ::= HAVING expr",
+ /* 144 */ "limit_opt ::=",
+ /* 145 */ "limit_opt ::= LIMIT expr",
+ /* 146 */ "limit_opt ::= LIMIT expr OFFSET expr",
+ /* 147 */ "limit_opt ::= LIMIT expr COMMA expr",
+ /* 148 */ "cmd ::= with DELETE FROM xfullname indexed_opt where_opt",
+ /* 149 */ "where_opt ::=",
+ /* 150 */ "where_opt ::= WHERE expr",
+ /* 151 */ "cmd ::= with UPDATE orconf xfullname indexed_opt SET setlist where_opt",
+ /* 152 */ "setlist ::= setlist COMMA nm EQ expr",
+ /* 153 */ "setlist ::= setlist COMMA LP idlist RP EQ expr",
+ /* 154 */ "setlist ::= nm EQ expr",
+ /* 155 */ "setlist ::= LP idlist RP EQ expr",
+ /* 156 */ "cmd ::= with insert_cmd INTO xfullname idlist_opt select upsert",
+ /* 157 */ "cmd ::= with insert_cmd INTO xfullname idlist_opt DEFAULT VALUES",
+ /* 158 */ "upsert ::=",
+ /* 159 */ "upsert ::= ON CONFLICT LP sortlist RP where_opt DO UPDATE SET setlist where_opt",
+ /* 160 */ "upsert ::= ON CONFLICT LP sortlist RP where_opt DO NOTHING",
+ /* 161 */ "upsert ::= ON CONFLICT DO NOTHING",
+ /* 162 */ "insert_cmd ::= INSERT orconf",
+ /* 163 */ "insert_cmd ::= REPLACE",
+ /* 164 */ "idlist_opt ::=",
+ /* 165 */ "idlist_opt ::= LP idlist RP",
+ /* 166 */ "idlist ::= idlist COMMA nm",
+ /* 167 */ "idlist ::= nm",
+ /* 168 */ "expr ::= LP expr RP",
+ /* 169 */ "expr ::= ID|INDEXED",
+ /* 170 */ "expr ::= JOIN_KW",
+ /* 171 */ "expr ::= nm DOT nm",
+ /* 172 */ "expr ::= nm DOT nm DOT nm",
+ /* 173 */ "term ::= NULL|FLOAT|BLOB",
+ /* 174 */ "term ::= STRING",
+ /* 175 */ "term ::= INTEGER",
+ /* 176 */ "expr ::= VARIABLE",
+ /* 177 */ "expr ::= expr COLLATE ID|STRING",
+ /* 178 */ "expr ::= CAST LP expr AS typetoken RP",
+ /* 179 */ "expr ::= ID|INDEXED LP distinct exprlist RP",
+ /* 180 */ "expr ::= ID|INDEXED LP STAR RP",
+ /* 181 */ "expr ::= ID|INDEXED LP distinct exprlist RP filter_over",
+ /* 182 */ "expr ::= ID|INDEXED LP STAR RP filter_over",
+ /* 183 */ "term ::= CTIME_KW",
+ /* 184 */ "expr ::= LP nexprlist COMMA expr RP",
+ /* 185 */ "expr ::= expr AND expr",
+ /* 186 */ "expr ::= expr OR expr",
+ /* 187 */ "expr ::= expr LT|GT|GE|LE expr",
+ /* 188 */ "expr ::= expr EQ|NE expr",
+ /* 189 */ "expr ::= expr BITAND|BITOR|LSHIFT|RSHIFT expr",
+ /* 190 */ "expr ::= expr PLUS|MINUS expr",
+ /* 191 */ "expr ::= expr STAR|SLASH|REM expr",
+ /* 192 */ "expr ::= expr CONCAT expr",
+ /* 193 */ "likeop ::= NOT LIKE_KW|MATCH",
+ /* 194 */ "expr ::= expr likeop expr",
+ /* 195 */ "expr ::= expr likeop expr ESCAPE expr",
+ /* 196 */ "expr ::= expr ISNULL|NOTNULL",
+ /* 197 */ "expr ::= expr NOT NULL",
+ /* 198 */ "expr ::= expr IS expr",
+ /* 199 */ "expr ::= expr IS NOT expr",
+ /* 200 */ "expr ::= NOT expr",
+ /* 201 */ "expr ::= BITNOT expr",
+ /* 202 */ "expr ::= PLUS|MINUS expr",
+ /* 203 */ "between_op ::= BETWEEN",
+ /* 204 */ "between_op ::= NOT BETWEEN",
+ /* 205 */ "expr ::= expr between_op expr AND expr",
+ /* 206 */ "in_op ::= IN",
+ /* 207 */ "in_op ::= NOT IN",
+ /* 208 */ "expr ::= expr in_op LP exprlist RP",
+ /* 209 */ "expr ::= LP select RP",
+ /* 210 */ "expr ::= expr in_op LP select RP",
+ /* 211 */ "expr ::= expr in_op nm dbnm paren_exprlist",
+ /* 212 */ "expr ::= EXISTS LP select RP",
+ /* 213 */ "expr ::= CASE case_operand case_exprlist case_else END",
+ /* 214 */ "case_exprlist ::= case_exprlist WHEN expr THEN expr",
+ /* 215 */ "case_exprlist ::= WHEN expr THEN expr",
+ /* 216 */ "case_else ::= ELSE expr",
+ /* 217 */ "case_else ::=",
+ /* 218 */ "case_operand ::= expr",
+ /* 219 */ "case_operand ::=",
+ /* 220 */ "exprlist ::=",
+ /* 221 */ "nexprlist ::= nexprlist COMMA expr",
+ /* 222 */ "nexprlist ::= expr",
+ /* 223 */ "paren_exprlist ::=",
+ /* 224 */ "paren_exprlist ::= LP exprlist RP",
+ /* 225 */ "cmd ::= createkw uniqueflag INDEX ifnotexists nm dbnm ON nm LP sortlist RP where_opt",
+ /* 226 */ "uniqueflag ::= UNIQUE",
+ /* 227 */ "uniqueflag ::=",
+ /* 228 */ "eidlist_opt ::=",
+ /* 229 */ "eidlist_opt ::= LP eidlist RP",
+ /* 230 */ "eidlist ::= eidlist COMMA nm collate sortorder",
+ /* 231 */ "eidlist ::= nm collate sortorder",
+ /* 232 */ "collate ::=",
+ /* 233 */ "collate ::= COLLATE ID|STRING",
+ /* 234 */ "cmd ::= DROP INDEX ifexists fullname",
+ /* 235 */ "cmd ::= VACUUM vinto",
+ /* 236 */ "cmd ::= VACUUM nm vinto",
+ /* 237 */ "vinto ::= INTO expr",
+ /* 238 */ "vinto ::=",
+ /* 239 */ "cmd ::= PRAGMA nm dbnm",
+ /* 240 */ "cmd ::= PRAGMA nm dbnm EQ nmnum",
+ /* 241 */ "cmd ::= PRAGMA nm dbnm LP nmnum RP",
+ /* 242 */ "cmd ::= PRAGMA nm dbnm EQ minus_num",
+ /* 243 */ "cmd ::= PRAGMA nm dbnm LP minus_num RP",
+ /* 244 */ "plus_num ::= PLUS INTEGER|FLOAT",
+ /* 245 */ "minus_num ::= MINUS INTEGER|FLOAT",
+ /* 246 */ "cmd ::= createkw trigger_decl BEGIN trigger_cmd_list END",
+ /* 247 */ "trigger_decl ::= temp TRIGGER ifnotexists nm dbnm trigger_time trigger_event ON fullname foreach_clause when_clause",
+ /* 248 */ "trigger_time ::= BEFORE|AFTER",
+ /* 249 */ "trigger_time ::= INSTEAD OF",
+ /* 250 */ "trigger_time ::=",
+ /* 251 */ "trigger_event ::= DELETE|INSERT",
+ /* 252 */ "trigger_event ::= UPDATE",
+ /* 253 */ "trigger_event ::= UPDATE OF idlist",
+ /* 254 */ "when_clause ::=",
+ /* 255 */ "when_clause ::= WHEN expr",
+ /* 256 */ "trigger_cmd_list ::= trigger_cmd_list trigger_cmd SEMI",
+ /* 257 */ "trigger_cmd_list ::= trigger_cmd SEMI",
+ /* 258 */ "trnm ::= nm DOT nm",
+ /* 259 */ "tridxby ::= INDEXED BY nm",
+ /* 260 */ "tridxby ::= NOT INDEXED",
+ /* 261 */ "trigger_cmd ::= UPDATE orconf trnm tridxby SET setlist where_opt scanpt",
+ /* 262 */ "trigger_cmd ::= scanpt insert_cmd INTO trnm idlist_opt select upsert scanpt",
+ /* 263 */ "trigger_cmd ::= DELETE FROM trnm tridxby where_opt scanpt",
+ /* 264 */ "trigger_cmd ::= scanpt select scanpt",
+ /* 265 */ "expr ::= RAISE LP IGNORE RP",
+ /* 266 */ "expr ::= RAISE LP raisetype COMMA nm RP",
+ /* 267 */ "raisetype ::= ROLLBACK",
+ /* 268 */ "raisetype ::= ABORT",
+ /* 269 */ "raisetype ::= FAIL",
+ /* 270 */ "cmd ::= DROP TRIGGER ifexists fullname",
+ /* 271 */ "cmd ::= ATTACH database_kw_opt expr AS expr key_opt",
+ /* 272 */ "cmd ::= DETACH database_kw_opt expr",
+ /* 273 */ "key_opt ::=",
+ /* 274 */ "key_opt ::= KEY expr",
+ /* 275 */ "cmd ::= REINDEX",
+ /* 276 */ "cmd ::= REINDEX nm dbnm",
+ /* 277 */ "cmd ::= ANALYZE",
+ /* 278 */ "cmd ::= ANALYZE nm dbnm",
+ /* 279 */ "cmd ::= ALTER TABLE fullname RENAME TO nm",
+ /* 280 */ "cmd ::= ALTER TABLE add_column_fullname ADD kwcolumn_opt columnname carglist",
+ /* 281 */ "add_column_fullname ::= fullname",
+ /* 282 */ "cmd ::= ALTER TABLE fullname RENAME kwcolumn_opt nm TO nm",
+ /* 283 */ "cmd ::= create_vtab",
+ /* 284 */ "cmd ::= create_vtab LP vtabarglist RP",
+ /* 285 */ "create_vtab ::= createkw VIRTUAL TABLE ifnotexists nm dbnm USING nm",
+ /* 286 */ "vtabarg ::=",
+ /* 287 */ "vtabargtoken ::= ANY",
+ /* 288 */ "vtabargtoken ::= lp anylist RP",
+ /* 289 */ "lp ::= LP",
+ /* 290 */ "with ::= WITH wqlist",
+ /* 291 */ "with ::= WITH RECURSIVE wqlist",
+ /* 292 */ "wqlist ::= nm eidlist_opt AS LP select RP",
+ /* 293 */ "wqlist ::= wqlist COMMA nm eidlist_opt AS LP select RP",
+ /* 294 */ "windowdefn_list ::= windowdefn",
+ /* 295 */ "windowdefn_list ::= windowdefn_list COMMA windowdefn",
+ /* 296 */ "windowdefn ::= nm AS LP window RP",
+ /* 297 */ "window ::= PARTITION BY nexprlist orderby_opt frame_opt",
+ /* 298 */ "window ::= nm PARTITION BY nexprlist orderby_opt frame_opt",
+ /* 299 */ "window ::= ORDER BY sortlist frame_opt",
+ /* 300 */ "window ::= nm ORDER BY sortlist frame_opt",
+ /* 301 */ "window ::= frame_opt",
+ /* 302 */ "window ::= nm frame_opt",
+ /* 303 */ "frame_opt ::=",
+ /* 304 */ "frame_opt ::= range_or_rows frame_bound_s frame_exclude_opt",
+ /* 305 */ "frame_opt ::= range_or_rows BETWEEN frame_bound_s AND frame_bound_e frame_exclude_opt",
+ /* 306 */ "range_or_rows ::= RANGE|ROWS|GROUPS",
+ /* 307 */ "frame_bound_s ::= frame_bound",
+ /* 308 */ "frame_bound_s ::= UNBOUNDED PRECEDING",
+ /* 309 */ "frame_bound_e ::= frame_bound",
+ /* 310 */ "frame_bound_e ::= UNBOUNDED FOLLOWING",
+ /* 311 */ "frame_bound ::= expr PRECEDING|FOLLOWING",
+ /* 312 */ "frame_bound ::= CURRENT ROW",
+ /* 313 */ "frame_exclude_opt ::=",
+ /* 314 */ "frame_exclude_opt ::= EXCLUDE frame_exclude",
+ /* 315 */ "frame_exclude ::= NO OTHERS",
+ /* 316 */ "frame_exclude ::= CURRENT ROW",
+ /* 317 */ "frame_exclude ::= GROUP|TIES",
+ /* 318 */ "window_clause ::= WINDOW windowdefn_list",
+ /* 319 */ "filter_over ::= filter_clause over_clause",
+ /* 320 */ "filter_over ::= over_clause",
+ /* 321 */ "filter_over ::= filter_clause",
+ /* 322 */ "over_clause ::= OVER LP window RP",
+ /* 323 */ "over_clause ::= OVER nm",
+ /* 324 */ "filter_clause ::= FILTER LP WHERE expr RP",
+ /* 325 */ "input ::= cmdlist",
+ /* 326 */ "cmdlist ::= cmdlist ecmd",
+ /* 327 */ "cmdlist ::= ecmd",
+ /* 328 */ "ecmd ::= SEMI",
+ /* 329 */ "ecmd ::= cmdx SEMI",
+ /* 330 */ "ecmd ::= explain cmdx SEMI",
+ /* 331 */ "trans_opt ::=",
+ /* 332 */ "trans_opt ::= TRANSACTION",
+ /* 333 */ "trans_opt ::= TRANSACTION nm",
+ /* 334 */ "savepoint_opt ::= SAVEPOINT",
+ /* 335 */ "savepoint_opt ::=",
+ /* 336 */ "cmd ::= create_table create_table_args",
+ /* 337 */ "columnlist ::= columnlist COMMA columnname carglist",
+ /* 338 */ "columnlist ::= columnname carglist",
+ /* 339 */ "nm ::= ID|INDEXED",
+ /* 340 */ "nm ::= STRING",
+ /* 341 */ "nm ::= JOIN_KW",
+ /* 342 */ "typetoken ::= typename",
+ /* 343 */ "typename ::= ID|STRING",
+ /* 344 */ "signed ::= plus_num",
+ /* 345 */ "signed ::= minus_num",
+ /* 346 */ "carglist ::= carglist ccons",
+ /* 347 */ "carglist ::=",
+ /* 348 */ "ccons ::= NULL onconf",
+ /* 349 */ "ccons ::= GENERATED ALWAYS AS generated",
+ /* 350 */ "ccons ::= AS generated",
+ /* 351 */ "conslist_opt ::= COMMA conslist",
+ /* 352 */ "conslist ::= conslist tconscomma tcons",
+ /* 353 */ "conslist ::= tcons",
+ /* 354 */ "tconscomma ::=",
+ /* 355 */ "defer_subclause_opt ::= defer_subclause",
+ /* 356 */ "resolvetype ::= raisetype",
+ /* 357 */ "selectnowith ::= oneselect",
+ /* 358 */ "oneselect ::= values",
+ /* 359 */ "sclp ::= selcollist COMMA",
+ /* 360 */ "as ::= ID|STRING",
+ /* 361 */ "expr ::= term",
+ /* 362 */ "likeop ::= LIKE_KW|MATCH",
+ /* 363 */ "exprlist ::= nexprlist",
+ /* 364 */ "nmnum ::= plus_num",
+ /* 365 */ "nmnum ::= nm",
+ /* 366 */ "nmnum ::= ON",
+ /* 367 */ "nmnum ::= DELETE",
+ /* 368 */ "nmnum ::= DEFAULT",
+ /* 369 */ "plus_num ::= INTEGER|FLOAT",
+ /* 370 */ "foreach_clause ::=",
+ /* 371 */ "foreach_clause ::= FOR EACH ROW",
+ /* 372 */ "trnm ::= nm",
+ /* 373 */ "tridxby ::=",
+ /* 374 */ "database_kw_opt ::= DATABASE",
+ /* 375 */ "database_kw_opt ::=",
+ /* 376 */ "kwcolumn_opt ::=",
+ /* 377 */ "kwcolumn_opt ::= COLUMNKW",
+ /* 378 */ "vtabarglist ::= vtabarg",
+ /* 379 */ "vtabarglist ::= vtabarglist COMMA vtabarg",
+ /* 380 */ "vtabarg ::= vtabarg vtabargtoken",
+ /* 381 */ "anylist ::=",
+ /* 382 */ "anylist ::= anylist LP anylist RP",
+ /* 383 */ "anylist ::= anylist ANY",
+ /* 384 */ "with ::=",
 };
 #endif /* NDEBUG */
 
@@ -150085,97 +154359,98 @@ static void yy_destructor(
     ** inside the C code.
     */
 /********* Begin destructor definitions ***************************************/
-    case 195: /* select */
-    case 227: /* selectnowith */
-    case 228: /* oneselect */
-    case 240: /* values */
+    case 200: /* select */
+    case 234: /* selectnowith */
+    case 235: /* oneselect */
+    case 247: /* values */
 {
-sqlite3SelectDelete(pParse->db, (yypminor->yy457));
+sqlite3SelectDelete(pParse->db, (yypminor->yy539));
 }
       break;
-    case 205: /* term */
-    case 206: /* expr */
-    case 234: /* where_opt */
-    case 236: /* having_opt */
-    case 248: /* on_opt */
-    case 263: /* case_operand */
-    case 265: /* case_else */
-    case 268: /* vinto */
-    case 275: /* when_clause */
-    case 280: /* key_opt */
-    case 294: /* filter_opt */
+    case 211: /* term */
+    case 212: /* expr */
+    case 241: /* where_opt */
+    case 243: /* having_opt */
+    case 255: /* on_opt */
+    case 271: /* case_operand */
+    case 273: /* case_else */
+    case 276: /* vinto */
+    case 283: /* when_clause */
+    case 288: /* key_opt */
+    case 302: /* filter_clause */
 {
-sqlite3ExprDelete(pParse->db, (yypminor->yy524));
+sqlite3ExprDelete(pParse->db, (yypminor->yy202));
 }
       break;
-    case 210: /* eidlist_opt */
-    case 219: /* sortlist */
-    case 220: /* eidlist */
-    case 232: /* selcollist */
-    case 235: /* groupby_opt */
-    case 237: /* orderby_opt */
-    case 241: /* nexprlist */
-    case 242: /* sclp */
-    case 250: /* exprlist */
-    case 254: /* setlist */
-    case 262: /* paren_exprlist */
-    case 264: /* case_exprlist */
-    case 293: /* part_opt */
+    case 216: /* eidlist_opt */
+    case 226: /* sortlist */
+    case 227: /* eidlist */
+    case 239: /* selcollist */
+    case 242: /* groupby_opt */
+    case 244: /* orderby_opt */
+    case 248: /* nexprlist */
+    case 249: /* sclp */
+    case 257: /* exprlist */
+    case 262: /* setlist */
+    case 270: /* paren_exprlist */
+    case 272: /* case_exprlist */
+    case 301: /* part_opt */
 {
-sqlite3ExprListDelete(pParse->db, (yypminor->yy434));
+sqlite3ExprListDelete(pParse->db, (yypminor->yy242));
 }
       break;
-    case 226: /* fullname */
-    case 233: /* from */
-    case 244: /* seltablist */
-    case 245: /* stl_prefix */
-    case 251: /* xfullname */
+    case 233: /* fullname */
+    case 240: /* from */
+    case 251: /* seltablist */
+    case 252: /* stl_prefix */
+    case 258: /* xfullname */
 {
-sqlite3SrcListDelete(pParse->db, (yypminor->yy483));
+sqlite3SrcListDelete(pParse->db, (yypminor->yy47));
 }
       break;
-    case 229: /* wqlist */
+    case 236: /* wqlist */
 {
-sqlite3WithDelete(pParse->db, (yypminor->yy59));
+sqlite3WithDelete(pParse->db, (yypminor->yy131));
 }
       break;
-    case 239: /* window_clause */
-    case 289: /* windowdefn_list */
+    case 246: /* window_clause */
+    case 297: /* windowdefn_list */
 {
-sqlite3WindowListDelete(pParse->db, (yypminor->yy295));
+sqlite3WindowListDelete(pParse->db, (yypminor->yy303));
 }
       break;
-    case 249: /* using_opt */
-    case 252: /* idlist */
-    case 256: /* idlist_opt */
+    case 256: /* using_opt */
+    case 259: /* idlist */
+    case 264: /* idlist_opt */
 {
-sqlite3IdListDelete(pParse->db, (yypminor->yy62));
+sqlite3IdListDelete(pParse->db, (yypminor->yy600));
 }
       break;
-    case 258: /* over_clause */
-    case 290: /* windowdefn */
-    case 291: /* window */
-    case 292: /* frame_opt */
+    case 266: /* filter_over */
+    case 298: /* windowdefn */
+    case 299: /* window */
+    case 300: /* frame_opt */
+    case 303: /* over_clause */
 {
-sqlite3WindowDelete(pParse->db, (yypminor->yy295));
+sqlite3WindowDelete(pParse->db, (yypminor->yy303));
 }
       break;
-    case 271: /* trigger_cmd_list */
-    case 276: /* trigger_cmd */
+    case 279: /* trigger_cmd_list */
+    case 284: /* trigger_cmd */
 {
-sqlite3DeleteTriggerStep(pParse->db, (yypminor->yy455));
+sqlite3DeleteTriggerStep(pParse->db, (yypminor->yy447));
 }
       break;
-    case 273: /* trigger_event */
+    case 281: /* trigger_event */
 {
-sqlite3IdListDelete(pParse->db, (yypminor->yy90).b);
+sqlite3IdListDelete(pParse->db, (yypminor->yy230).b);
 }
       break;
-    case 296: /* frame_bound */
-    case 297: /* frame_bound_s */
-    case 298: /* frame_bound_e */
+    case 305: /* frame_bound */
+    case 306: /* frame_bound_s */
+    case 307: /* frame_bound_e */
 {
-sqlite3ExprDelete(pParse->db, (yypminor->yy201).pExpr);
+sqlite3ExprDelete(pParse->db, (yypminor->yy77).pExpr);
 }
       break;
 /********* End destructor definitions *****************************************/
@@ -150301,15 +154576,18 @@ static YYACTIONTYPE yy_find_shift_action(
   do{
     i = yy_shift_ofst[stateno];
     assert( i>=0 );
-    /* assert( i+YYNTOKEN<=(int)YY_NLOOKAHEAD ); */
+    assert( i<=YY_ACTTAB_COUNT );
+    assert( i+YYNTOKEN<=(int)YY_NLOOKAHEAD );
     assert( iLookAhead!=YYNOCODE );
     assert( iLookAhead < YYNTOKEN );
     i += iLookAhead;
-    if( i>=YY_NLOOKAHEAD || yy_lookahead[i]!=iLookAhead ){
+    assert( i<(int)YY_NLOOKAHEAD );
+    if( yy_lookahead[i]!=iLookAhead ){
 #ifdef YYFALLBACK
       YYCODETYPE iFallback;            /* Fallback token */
-      if( iLookAhead<sizeof(yyFallback)/sizeof(yyFallback[0])
-             && (iFallback = yyFallback[iLookAhead])!=0 ){
+      assert( iLookAhead<sizeof(yyFallback)/sizeof(yyFallback[0]) );
+      iFallback = yyFallback[iLookAhead];
+      if( iFallback!=0 ){
 #ifndef NDEBUG
         if( yyTraceFILE ){
           fprintf(yyTraceFILE, "%sFALLBACK %s => %s\n",
@@ -150324,16 +154602,8 @@ static YYACTIONTYPE yy_find_shift_action(
 #ifdef YYWILDCARD
       {
         int j = i - iLookAhead + YYWILDCARD;
-        if( 
-#if YY_SHIFT_MIN+YYWILDCARD<0
-          j>=0 &&
-#endif
-#if YY_SHIFT_MAX+YYWILDCARD>=YY_ACTTAB_COUNT
-          j<YY_ACTTAB_COUNT &&
-#endif
-          j<(int)(sizeof(yy_lookahead)/sizeof(yy_lookahead[0])) &&
-          yy_lookahead[j]==YYWILDCARD && iLookAhead>0
-        ){
+        assert( j<(int)(sizeof(yy_lookahead)/sizeof(yy_lookahead[0])) );
+        if( yy_lookahead[j]==YYWILDCARD && iLookAhead>0 ){
 #ifndef NDEBUG
           if( yyTraceFILE ){
             fprintf(yyTraceFILE, "%sWILDCARD %s => %s\n",
@@ -150347,6 +154617,7 @@ static YYACTIONTYPE yy_find_shift_action(
 #endif /* YYWILDCARD */
       return yy_default[stateno];
     }else{
+      assert( i>=0 && i<sizeof(yy_action)/sizeof(yy_action[0]) );
       return yy_action[i];
     }
   }while(1);
@@ -150470,381 +154741,391 @@ static void yy_shift(
 /* For rule J, yyRuleInfoLhs[J] contains the symbol on the left-hand side
 ** of that rule */
 static const YYCODETYPE yyRuleInfoLhs[] = {
-   180,  /* (0) explain ::= EXPLAIN */
-   180,  /* (1) explain ::= EXPLAIN QUERY PLAN */
-   179,  /* (2) cmdx ::= cmd */
-   181,  /* (3) cmd ::= BEGIN transtype trans_opt */
-   182,  /* (4) transtype ::= */
-   182,  /* (5) transtype ::= DEFERRED */
-   182,  /* (6) transtype ::= IMMEDIATE */
-   182,  /* (7) transtype ::= EXCLUSIVE */
-   181,  /* (8) cmd ::= COMMIT|END trans_opt */
-   181,  /* (9) cmd ::= ROLLBACK trans_opt */
-   181,  /* (10) cmd ::= SAVEPOINT nm */
-   181,  /* (11) cmd ::= RELEASE savepoint_opt nm */
-   181,  /* (12) cmd ::= ROLLBACK trans_opt TO savepoint_opt nm */
-   186,  /* (13) create_table ::= createkw temp TABLE ifnotexists nm dbnm */
-   188,  /* (14) createkw ::= CREATE */
-   190,  /* (15) ifnotexists ::= */
-   190,  /* (16) ifnotexists ::= IF NOT EXISTS */
-   189,  /* (17) temp ::= TEMP */
-   189,  /* (18) temp ::= */
-   187,  /* (19) create_table_args ::= LP columnlist conslist_opt RP table_options */
-   187,  /* (20) create_table_args ::= AS select */
-   194,  /* (21) table_options ::= */
-   194,  /* (22) table_options ::= WITHOUT nm */
-   196,  /* (23) columnname ::= nm typetoken */
-   198,  /* (24) typetoken ::= */
-   198,  /* (25) typetoken ::= typename LP signed RP */
-   198,  /* (26) typetoken ::= typename LP signed COMMA signed RP */
-   199,  /* (27) typename ::= typename ID|STRING */
-   203,  /* (28) scanpt ::= */
-   204,  /* (29) ccons ::= CONSTRAINT nm */
-   204,  /* (30) ccons ::= DEFAULT scanpt term scanpt */
-   204,  /* (31) ccons ::= DEFAULT LP expr RP */
-   204,  /* (32) ccons ::= DEFAULT PLUS term scanpt */
-   204,  /* (33) ccons ::= DEFAULT MINUS term scanpt */
-   204,  /* (34) ccons ::= DEFAULT scanpt ID|INDEXED */
-   204,  /* (35) ccons ::= NOT NULL onconf */
-   204,  /* (36) ccons ::= PRIMARY KEY sortorder onconf autoinc */
-   204,  /* (37) ccons ::= UNIQUE onconf */
-   204,  /* (38) ccons ::= CHECK LP expr RP */
-   204,  /* (39) ccons ::= REFERENCES nm eidlist_opt refargs */
-   204,  /* (40) ccons ::= defer_subclause */
-   204,  /* (41) ccons ::= COLLATE ID|STRING */
-   209,  /* (42) autoinc ::= */
-   209,  /* (43) autoinc ::= AUTOINCR */
-   211,  /* (44) refargs ::= */
-   211,  /* (45) refargs ::= refargs refarg */
-   213,  /* (46) refarg ::= MATCH nm */
-   213,  /* (47) refarg ::= ON INSERT refact */
-   213,  /* (48) refarg ::= ON DELETE refact */
-   213,  /* (49) refarg ::= ON UPDATE refact */
-   214,  /* (50) refact ::= SET NULL */
-   214,  /* (51) refact ::= SET DEFAULT */
-   214,  /* (52) refact ::= CASCADE */
-   214,  /* (53) refact ::= RESTRICT */
-   214,  /* (54) refact ::= NO ACTION */
-   212,  /* (55) defer_subclause ::= NOT DEFERRABLE init_deferred_pred_opt */
-   212,  /* (56) defer_subclause ::= DEFERRABLE init_deferred_pred_opt */
-   215,  /* (57) init_deferred_pred_opt ::= */
-   215,  /* (58) init_deferred_pred_opt ::= INITIALLY DEFERRED */
-   215,  /* (59) init_deferred_pred_opt ::= INITIALLY IMMEDIATE */
-   193,  /* (60) conslist_opt ::= */
-   217,  /* (61) tconscomma ::= COMMA */
-   218,  /* (62) tcons ::= CONSTRAINT nm */
-   218,  /* (63) tcons ::= PRIMARY KEY LP sortlist autoinc RP onconf */
-   218,  /* (64) tcons ::= UNIQUE LP sortlist RP onconf */
-   218,  /* (65) tcons ::= CHECK LP expr RP onconf */
-   218,  /* (66) tcons ::= FOREIGN KEY LP eidlist RP REFERENCES nm eidlist_opt refargs defer_subclause_opt */
-   221,  /* (67) defer_subclause_opt ::= */
-   207,  /* (68) onconf ::= */
-   207,  /* (69) onconf ::= ON CONFLICT resolvetype */
-   222,  /* (70) orconf ::= */
-   222,  /* (71) orconf ::= OR resolvetype */
-   223,  /* (72) resolvetype ::= IGNORE */
-   223,  /* (73) resolvetype ::= REPLACE */
-   181,  /* (74) cmd ::= DROP TABLE ifexists fullname */
-   225,  /* (75) ifexists ::= IF EXISTS */
-   225,  /* (76) ifexists ::= */
-   181,  /* (77) cmd ::= createkw temp VIEW ifnotexists nm dbnm eidlist_opt AS select */
-   181,  /* (78) cmd ::= DROP VIEW ifexists fullname */
-   181,  /* (79) cmd ::= select */
-   195,  /* (80) select ::= WITH wqlist selectnowith */
-   195,  /* (81) select ::= WITH RECURSIVE wqlist selectnowith */
-   195,  /* (82) select ::= selectnowith */
-   227,  /* (83) selectnowith ::= selectnowith multiselect_op oneselect */
-   230,  /* (84) multiselect_op ::= UNION */
-   230,  /* (85) multiselect_op ::= UNION ALL */
-   230,  /* (86) multiselect_op ::= EXCEPT|INTERSECT */
-   228,  /* (87) oneselect ::= SELECT distinct selcollist from where_opt groupby_opt having_opt orderby_opt limit_opt */
-   228,  /* (88) oneselect ::= SELECT distinct selcollist from where_opt groupby_opt having_opt window_clause orderby_opt limit_opt */
-   240,  /* (89) values ::= VALUES LP nexprlist RP */
-   240,  /* (90) values ::= values COMMA LP nexprlist RP */
-   231,  /* (91) distinct ::= DISTINCT */
-   231,  /* (92) distinct ::= ALL */
-   231,  /* (93) distinct ::= */
-   242,  /* (94) sclp ::= */
-   232,  /* (95) selcollist ::= sclp scanpt expr scanpt as */
-   232,  /* (96) selcollist ::= sclp scanpt STAR */
-   232,  /* (97) selcollist ::= sclp scanpt nm DOT STAR */
-   243,  /* (98) as ::= AS nm */
-   243,  /* (99) as ::= */
-   233,  /* (100) from ::= */
-   233,  /* (101) from ::= FROM seltablist */
-   245,  /* (102) stl_prefix ::= seltablist joinop */
-   245,  /* (103) stl_prefix ::= */
-   244,  /* (104) seltablist ::= stl_prefix nm dbnm as indexed_opt on_opt using_opt */
-   244,  /* (105) seltablist ::= stl_prefix nm dbnm LP exprlist RP as on_opt using_opt */
-   244,  /* (106) seltablist ::= stl_prefix LP select RP as on_opt using_opt */
-   244,  /* (107) seltablist ::= stl_prefix LP seltablist RP as on_opt using_opt */
-   191,  /* (108) dbnm ::= */
-   191,  /* (109) dbnm ::= DOT nm */
-   226,  /* (110) fullname ::= nm */
-   226,  /* (111) fullname ::= nm DOT nm */
-   251,  /* (112) xfullname ::= nm */
-   251,  /* (113) xfullname ::= nm DOT nm */
-   251,  /* (114) xfullname ::= nm DOT nm AS nm */
-   251,  /* (115) xfullname ::= nm AS nm */
-   246,  /* (116) joinop ::= COMMA|JOIN */
-   246,  /* (117) joinop ::= JOIN_KW JOIN */
-   246,  /* (118) joinop ::= JOIN_KW nm JOIN */
-   246,  /* (119) joinop ::= JOIN_KW nm nm JOIN */
-   248,  /* (120) on_opt ::= ON expr */
-   248,  /* (121) on_opt ::= */
-   247,  /* (122) indexed_opt ::= */
-   247,  /* (123) indexed_opt ::= INDEXED BY nm */
-   247,  /* (124) indexed_opt ::= NOT INDEXED */
-   249,  /* (125) using_opt ::= USING LP idlist RP */
-   249,  /* (126) using_opt ::= */
-   237,  /* (127) orderby_opt ::= */
-   237,  /* (128) orderby_opt ::= ORDER BY sortlist */
-   219,  /* (129) sortlist ::= sortlist COMMA expr sortorder */
-   219,  /* (130) sortlist ::= expr sortorder */
-   208,  /* (131) sortorder ::= ASC */
-   208,  /* (132) sortorder ::= DESC */
-   208,  /* (133) sortorder ::= */
-   235,  /* (134) groupby_opt ::= */
-   235,  /* (135) groupby_opt ::= GROUP BY nexprlist */
-   236,  /* (136) having_opt ::= */
-   236,  /* (137) having_opt ::= HAVING expr */
-   238,  /* (138) limit_opt ::= */
-   238,  /* (139) limit_opt ::= LIMIT expr */
-   238,  /* (140) limit_opt ::= LIMIT expr OFFSET expr */
-   238,  /* (141) limit_opt ::= LIMIT expr COMMA expr */
-   181,  /* (142) cmd ::= with DELETE FROM xfullname indexed_opt where_opt */
-   234,  /* (143) where_opt ::= */
-   234,  /* (144) where_opt ::= WHERE expr */
-   181,  /* (145) cmd ::= with UPDATE orconf xfullname indexed_opt SET setlist where_opt */
-   254,  /* (146) setlist ::= setlist COMMA nm EQ expr */
-   254,  /* (147) setlist ::= setlist COMMA LP idlist RP EQ expr */
-   254,  /* (148) setlist ::= nm EQ expr */
-   254,  /* (149) setlist ::= LP idlist RP EQ expr */
-   181,  /* (150) cmd ::= with insert_cmd INTO xfullname idlist_opt select upsert */
-   181,  /* (151) cmd ::= with insert_cmd INTO xfullname idlist_opt DEFAULT VALUES */
-   257,  /* (152) upsert ::= */
-   257,  /* (153) upsert ::= ON CONFLICT LP sortlist RP where_opt DO UPDATE SET setlist where_opt */
-   257,  /* (154) upsert ::= ON CONFLICT LP sortlist RP where_opt DO NOTHING */
-   257,  /* (155) upsert ::= ON CONFLICT DO NOTHING */
-   255,  /* (156) insert_cmd ::= INSERT orconf */
-   255,  /* (157) insert_cmd ::= REPLACE */
-   256,  /* (158) idlist_opt ::= */
-   256,  /* (159) idlist_opt ::= LP idlist RP */
-   252,  /* (160) idlist ::= idlist COMMA nm */
-   252,  /* (161) idlist ::= nm */
-   206,  /* (162) expr ::= LP expr RP */
-   206,  /* (163) expr ::= ID|INDEXED */
-   206,  /* (164) expr ::= JOIN_KW */
-   206,  /* (165) expr ::= nm DOT nm */
-   206,  /* (166) expr ::= nm DOT nm DOT nm */
-   205,  /* (167) term ::= NULL|FLOAT|BLOB */
-   205,  /* (168) term ::= STRING */
-   205,  /* (169) term ::= INTEGER */
-   206,  /* (170) expr ::= VARIABLE */
-   206,  /* (171) expr ::= expr COLLATE ID|STRING */
-   206,  /* (172) expr ::= CAST LP expr AS typetoken RP */
-   206,  /* (173) expr ::= ID|INDEXED LP distinct exprlist RP */
-   206,  /* (174) expr ::= ID|INDEXED LP STAR RP */
-   206,  /* (175) expr ::= ID|INDEXED LP distinct exprlist RP over_clause */
-   206,  /* (176) expr ::= ID|INDEXED LP STAR RP over_clause */
-   205,  /* (177) term ::= CTIME_KW */
-   206,  /* (178) expr ::= LP nexprlist COMMA expr RP */
-   206,  /* (179) expr ::= expr AND expr */
-   206,  /* (180) expr ::= expr OR expr */
-   206,  /* (181) expr ::= expr LT|GT|GE|LE expr */
-   206,  /* (182) expr ::= expr EQ|NE expr */
-   206,  /* (183) expr ::= expr BITAND|BITOR|LSHIFT|RSHIFT expr */
-   206,  /* (184) expr ::= expr PLUS|MINUS expr */
-   206,  /* (185) expr ::= expr STAR|SLASH|REM expr */
-   206,  /* (186) expr ::= expr CONCAT expr */
-   259,  /* (187) likeop ::= NOT LIKE_KW|MATCH */
-   206,  /* (188) expr ::= expr likeop expr */
-   206,  /* (189) expr ::= expr likeop expr ESCAPE expr */
-   206,  /* (190) expr ::= expr ISNULL|NOTNULL */
-   206,  /* (191) expr ::= expr NOT NULL */
-   206,  /* (192) expr ::= expr IS expr */
-   206,  /* (193) expr ::= expr IS NOT expr */
-   206,  /* (194) expr ::= NOT expr */
-   206,  /* (195) expr ::= BITNOT expr */
-   206,  /* (196) expr ::= PLUS|MINUS expr */
-   260,  /* (197) between_op ::= BETWEEN */
-   260,  /* (198) between_op ::= NOT BETWEEN */
-   206,  /* (199) expr ::= expr between_op expr AND expr */
-   261,  /* (200) in_op ::= IN */
-   261,  /* (201) in_op ::= NOT IN */
-   206,  /* (202) expr ::= expr in_op LP exprlist RP */
-   206,  /* (203) expr ::= LP select RP */
-   206,  /* (204) expr ::= expr in_op LP select RP */
-   206,  /* (205) expr ::= expr in_op nm dbnm paren_exprlist */
-   206,  /* (206) expr ::= EXISTS LP select RP */
-   206,  /* (207) expr ::= CASE case_operand case_exprlist case_else END */
-   264,  /* (208) case_exprlist ::= case_exprlist WHEN expr THEN expr */
-   264,  /* (209) case_exprlist ::= WHEN expr THEN expr */
-   265,  /* (210) case_else ::= ELSE expr */
-   265,  /* (211) case_else ::= */
-   263,  /* (212) case_operand ::= expr */
-   263,  /* (213) case_operand ::= */
-   250,  /* (214) exprlist ::= */
-   241,  /* (215) nexprlist ::= nexprlist COMMA expr */
-   241,  /* (216) nexprlist ::= expr */
-   262,  /* (217) paren_exprlist ::= */
-   262,  /* (218) paren_exprlist ::= LP exprlist RP */
-   181,  /* (219) cmd ::= createkw uniqueflag INDEX ifnotexists nm dbnm ON nm LP sortlist RP where_opt */
-   266,  /* (220) uniqueflag ::= UNIQUE */
-   266,  /* (221) uniqueflag ::= */
-   210,  /* (222) eidlist_opt ::= */
-   210,  /* (223) eidlist_opt ::= LP eidlist RP */
-   220,  /* (224) eidlist ::= eidlist COMMA nm collate sortorder */
-   220,  /* (225) eidlist ::= nm collate sortorder */
-   267,  /* (226) collate ::= */
-   267,  /* (227) collate ::= COLLATE ID|STRING */
-   181,  /* (228) cmd ::= DROP INDEX ifexists fullname */
-   181,  /* (229) cmd ::= VACUUM vinto */
-   181,  /* (230) cmd ::= VACUUM nm vinto */
-   268,  /* (231) vinto ::= INTO expr */
-   268,  /* (232) vinto ::= */
-   181,  /* (233) cmd ::= PRAGMA nm dbnm */
-   181,  /* (234) cmd ::= PRAGMA nm dbnm EQ nmnum */
-   181,  /* (235) cmd ::= PRAGMA nm dbnm LP nmnum RP */
-   181,  /* (236) cmd ::= PRAGMA nm dbnm EQ minus_num */
-   181,  /* (237) cmd ::= PRAGMA nm dbnm LP minus_num RP */
-   201,  /* (238) plus_num ::= PLUS INTEGER|FLOAT */
-   202,  /* (239) minus_num ::= MINUS INTEGER|FLOAT */
-   181,  /* (240) cmd ::= createkw trigger_decl BEGIN trigger_cmd_list END */
-   270,  /* (241) trigger_decl ::= temp TRIGGER ifnotexists nm dbnm trigger_time trigger_event ON fullname foreach_clause when_clause */
-   272,  /* (242) trigger_time ::= BEFORE|AFTER */
-   272,  /* (243) trigger_time ::= INSTEAD OF */
-   272,  /* (244) trigger_time ::= */
-   273,  /* (245) trigger_event ::= DELETE|INSERT */
-   273,  /* (246) trigger_event ::= UPDATE */
-   273,  /* (247) trigger_event ::= UPDATE OF idlist */
-   275,  /* (248) when_clause ::= */
-   275,  /* (249) when_clause ::= WHEN expr */
-   271,  /* (250) trigger_cmd_list ::= trigger_cmd_list trigger_cmd SEMI */
-   271,  /* (251) trigger_cmd_list ::= trigger_cmd SEMI */
-   277,  /* (252) trnm ::= nm DOT nm */
-   278,  /* (253) tridxby ::= INDEXED BY nm */
-   278,  /* (254) tridxby ::= NOT INDEXED */
-   276,  /* (255) trigger_cmd ::= UPDATE orconf trnm tridxby SET setlist where_opt scanpt */
-   276,  /* (256) trigger_cmd ::= scanpt insert_cmd INTO trnm idlist_opt select upsert scanpt */
-   276,  /* (257) trigger_cmd ::= DELETE FROM trnm tridxby where_opt scanpt */
-   276,  /* (258) trigger_cmd ::= scanpt select scanpt */
-   206,  /* (259) expr ::= RAISE LP IGNORE RP */
-   206,  /* (260) expr ::= RAISE LP raisetype COMMA nm RP */
-   224,  /* (261) raisetype ::= ROLLBACK */
-   224,  /* (262) raisetype ::= ABORT */
-   224,  /* (263) raisetype ::= FAIL */
-   181,  /* (264) cmd ::= DROP TRIGGER ifexists fullname */
-   181,  /* (265) cmd ::= ATTACH database_kw_opt expr AS expr key_opt */
-   181,  /* (266) cmd ::= DETACH database_kw_opt expr */
-   280,  /* (267) key_opt ::= */
-   280,  /* (268) key_opt ::= KEY expr */
-   181,  /* (269) cmd ::= REINDEX */
-   181,  /* (270) cmd ::= REINDEX nm dbnm */
-   181,  /* (271) cmd ::= ANALYZE */
-   181,  /* (272) cmd ::= ANALYZE nm dbnm */
-   181,  /* (273) cmd ::= ALTER TABLE fullname RENAME TO nm */
-   181,  /* (274) cmd ::= ALTER TABLE add_column_fullname ADD kwcolumn_opt columnname carglist */
-   281,  /* (275) add_column_fullname ::= fullname */
-   181,  /* (276) cmd ::= ALTER TABLE fullname RENAME kwcolumn_opt nm TO nm */
-   181,  /* (277) cmd ::= create_vtab */
-   181,  /* (278) cmd ::= create_vtab LP vtabarglist RP */
-   283,  /* (279) create_vtab ::= createkw VIRTUAL TABLE ifnotexists nm dbnm USING nm */
-   285,  /* (280) vtabarg ::= */
-   286,  /* (281) vtabargtoken ::= ANY */
-   286,  /* (282) vtabargtoken ::= lp anylist RP */
-   287,  /* (283) lp ::= LP */
-   253,  /* (284) with ::= WITH wqlist */
-   253,  /* (285) with ::= WITH RECURSIVE wqlist */
-   229,  /* (286) wqlist ::= nm eidlist_opt AS LP select RP */
-   229,  /* (287) wqlist ::= wqlist COMMA nm eidlist_opt AS LP select RP */
-   289,  /* (288) windowdefn_list ::= windowdefn */
-   289,  /* (289) windowdefn_list ::= windowdefn_list COMMA windowdefn */
-   290,  /* (290) windowdefn ::= nm AS LP window RP */
-   291,  /* (291) window ::= PARTITION BY nexprlist orderby_opt frame_opt */
-   291,  /* (292) window ::= nm PARTITION BY nexprlist orderby_opt frame_opt */
-   291,  /* (293) window ::= ORDER BY sortlist frame_opt */
-   291,  /* (294) window ::= nm ORDER BY sortlist frame_opt */
-   291,  /* (295) window ::= frame_opt */
-   291,  /* (296) window ::= nm frame_opt */
-   292,  /* (297) frame_opt ::= */
-   292,  /* (298) frame_opt ::= range_or_rows frame_bound_s frame_exclude_opt */
-   292,  /* (299) frame_opt ::= range_or_rows BETWEEN frame_bound_s AND frame_bound_e frame_exclude_opt */
-   295,  /* (300) range_or_rows ::= RANGE|ROWS|GROUPS */
-   297,  /* (301) frame_bound_s ::= frame_bound */
-   297,  /* (302) frame_bound_s ::= UNBOUNDED PRECEDING */
-   298,  /* (303) frame_bound_e ::= frame_bound */
-   298,  /* (304) frame_bound_e ::= UNBOUNDED FOLLOWING */
-   296,  /* (305) frame_bound ::= expr PRECEDING|FOLLOWING */
-   296,  /* (306) frame_bound ::= CURRENT ROW */
-   299,  /* (307) frame_exclude_opt ::= */
-   299,  /* (308) frame_exclude_opt ::= EXCLUDE frame_exclude */
-   300,  /* (309) frame_exclude ::= NO OTHERS */
-   300,  /* (310) frame_exclude ::= CURRENT ROW */
-   300,  /* (311) frame_exclude ::= GROUP|TIES */
-   239,  /* (312) window_clause ::= WINDOW windowdefn_list */
-   258,  /* (313) over_clause ::= filter_opt OVER LP window RP */
-   258,  /* (314) over_clause ::= filter_opt OVER nm */
-   294,  /* (315) filter_opt ::= */
-   294,  /* (316) filter_opt ::= FILTER LP WHERE expr RP */
-   176,  /* (317) input ::= cmdlist */
-   177,  /* (318) cmdlist ::= cmdlist ecmd */
-   177,  /* (319) cmdlist ::= ecmd */
-   178,  /* (320) ecmd ::= SEMI */
-   178,  /* (321) ecmd ::= cmdx SEMI */
-   178,  /* (322) ecmd ::= explain cmdx */
-   183,  /* (323) trans_opt ::= */
-   183,  /* (324) trans_opt ::= TRANSACTION */
-   183,  /* (325) trans_opt ::= TRANSACTION nm */
-   185,  /* (326) savepoint_opt ::= SAVEPOINT */
-   185,  /* (327) savepoint_opt ::= */
-   181,  /* (328) cmd ::= create_table create_table_args */
-   192,  /* (329) columnlist ::= columnlist COMMA columnname carglist */
-   192,  /* (330) columnlist ::= columnname carglist */
-   184,  /* (331) nm ::= ID|INDEXED */
-   184,  /* (332) nm ::= STRING */
-   184,  /* (333) nm ::= JOIN_KW */
-   198,  /* (334) typetoken ::= typename */
-   199,  /* (335) typename ::= ID|STRING */
-   200,  /* (336) signed ::= plus_num */
-   200,  /* (337) signed ::= minus_num */
-   197,  /* (338) carglist ::= carglist ccons */
-   197,  /* (339) carglist ::= */
-   204,  /* (340) ccons ::= NULL onconf */
-   193,  /* (341) conslist_opt ::= COMMA conslist */
-   216,  /* (342) conslist ::= conslist tconscomma tcons */
-   216,  /* (343) conslist ::= tcons */
-   217,  /* (344) tconscomma ::= */
-   221,  /* (345) defer_subclause_opt ::= defer_subclause */
-   223,  /* (346) resolvetype ::= raisetype */
-   227,  /* (347) selectnowith ::= oneselect */
-   228,  /* (348) oneselect ::= values */
-   242,  /* (349) sclp ::= selcollist COMMA */
-   243,  /* (350) as ::= ID|STRING */
-   206,  /* (351) expr ::= term */
-   259,  /* (352) likeop ::= LIKE_KW|MATCH */
-   250,  /* (353) exprlist ::= nexprlist */
-   269,  /* (354) nmnum ::= plus_num */
-   269,  /* (355) nmnum ::= nm */
-   269,  /* (356) nmnum ::= ON */
-   269,  /* (357) nmnum ::= DELETE */
-   269,  /* (358) nmnum ::= DEFAULT */
-   201,  /* (359) plus_num ::= INTEGER|FLOAT */
-   274,  /* (360) foreach_clause ::= */
-   274,  /* (361) foreach_clause ::= FOR EACH ROW */
-   277,  /* (362) trnm ::= nm */
-   278,  /* (363) tridxby ::= */
-   279,  /* (364) database_kw_opt ::= DATABASE */
-   279,  /* (365) database_kw_opt ::= */
-   282,  /* (366) kwcolumn_opt ::= */
-   282,  /* (367) kwcolumn_opt ::= COLUMNKW */
-   284,  /* (368) vtabarglist ::= vtabarg */
-   284,  /* (369) vtabarglist ::= vtabarglist COMMA vtabarg */
-   285,  /* (370) vtabarg ::= vtabarg vtabargtoken */
-   288,  /* (371) anylist ::= */
-   288,  /* (372) anylist ::= anylist LP anylist RP */
-   288,  /* (373) anylist ::= anylist ANY */
-   253,  /* (374) with ::= */
+   185,  /* (0) explain ::= EXPLAIN */
+   185,  /* (1) explain ::= EXPLAIN QUERY PLAN */
+   184,  /* (2) cmdx ::= cmd */
+   186,  /* (3) cmd ::= BEGIN transtype trans_opt */
+   187,  /* (4) transtype ::= */
+   187,  /* (5) transtype ::= DEFERRED */
+   187,  /* (6) transtype ::= IMMEDIATE */
+   187,  /* (7) transtype ::= EXCLUSIVE */
+   186,  /* (8) cmd ::= COMMIT|END trans_opt */
+   186,  /* (9) cmd ::= ROLLBACK trans_opt */
+   186,  /* (10) cmd ::= SAVEPOINT nm */
+   186,  /* (11) cmd ::= RELEASE savepoint_opt nm */
+   186,  /* (12) cmd ::= ROLLBACK trans_opt TO savepoint_opt nm */
+   191,  /* (13) create_table ::= createkw temp TABLE ifnotexists nm dbnm */
+   193,  /* (14) createkw ::= CREATE */
+   195,  /* (15) ifnotexists ::= */
+   195,  /* (16) ifnotexists ::= IF NOT EXISTS */
+   194,  /* (17) temp ::= TEMP */
+   194,  /* (18) temp ::= */
+   192,  /* (19) create_table_args ::= LP columnlist conslist_opt RP table_options */
+   192,  /* (20) create_table_args ::= AS select */
+   199,  /* (21) table_options ::= */
+   199,  /* (22) table_options ::= WITHOUT nm */
+   201,  /* (23) columnname ::= nm typetoken */
+   203,  /* (24) typetoken ::= */
+   203,  /* (25) typetoken ::= typename LP signed RP */
+   203,  /* (26) typetoken ::= typename LP signed COMMA signed RP */
+   204,  /* (27) typename ::= typename ID|STRING */
+   208,  /* (28) scanpt ::= */
+   209,  /* (29) scantok ::= */
+   210,  /* (30) ccons ::= CONSTRAINT nm */
+   210,  /* (31) ccons ::= DEFAULT scantok term */
+   210,  /* (32) ccons ::= DEFAULT LP expr RP */
+   210,  /* (33) ccons ::= DEFAULT PLUS scantok term */
+   210,  /* (34) ccons ::= DEFAULT MINUS scantok term */
+   210,  /* (35) ccons ::= DEFAULT scantok ID|INDEXED */
+   210,  /* (36) ccons ::= NOT NULL onconf */
+   210,  /* (37) ccons ::= PRIMARY KEY sortorder onconf autoinc */
+   210,  /* (38) ccons ::= UNIQUE onconf */
+   210,  /* (39) ccons ::= CHECK LP expr RP */
+   210,  /* (40) ccons ::= REFERENCES nm eidlist_opt refargs */
+   210,  /* (41) ccons ::= defer_subclause */
+   210,  /* (42) ccons ::= COLLATE ID|STRING */
+   219,  /* (43) generated ::= LP expr RP */
+   219,  /* (44) generated ::= LP expr RP ID */
+   215,  /* (45) autoinc ::= */
+   215,  /* (46) autoinc ::= AUTOINCR */
+   217,  /* (47) refargs ::= */
+   217,  /* (48) refargs ::= refargs refarg */
+   220,  /* (49) refarg ::= MATCH nm */
+   220,  /* (50) refarg ::= ON INSERT refact */
+   220,  /* (51) refarg ::= ON DELETE refact */
+   220,  /* (52) refarg ::= ON UPDATE refact */
+   221,  /* (53) refact ::= SET NULL */
+   221,  /* (54) refact ::= SET DEFAULT */
+   221,  /* (55) refact ::= CASCADE */
+   221,  /* (56) refact ::= RESTRICT */
+   221,  /* (57) refact ::= NO ACTION */
+   218,  /* (58) defer_subclause ::= NOT DEFERRABLE init_deferred_pred_opt */
+   218,  /* (59) defer_subclause ::= DEFERRABLE init_deferred_pred_opt */
+   222,  /* (60) init_deferred_pred_opt ::= */
+   222,  /* (61) init_deferred_pred_opt ::= INITIALLY DEFERRED */
+   222,  /* (62) init_deferred_pred_opt ::= INITIALLY IMMEDIATE */
+   198,  /* (63) conslist_opt ::= */
+   224,  /* (64) tconscomma ::= COMMA */
+   225,  /* (65) tcons ::= CONSTRAINT nm */
+   225,  /* (66) tcons ::= PRIMARY KEY LP sortlist autoinc RP onconf */
+   225,  /* (67) tcons ::= UNIQUE LP sortlist RP onconf */
+   225,  /* (68) tcons ::= CHECK LP expr RP onconf */
+   225,  /* (69) tcons ::= FOREIGN KEY LP eidlist RP REFERENCES nm eidlist_opt refargs defer_subclause_opt */
+   228,  /* (70) defer_subclause_opt ::= */
+   213,  /* (71) onconf ::= */
+   213,  /* (72) onconf ::= ON CONFLICT resolvetype */
+   229,  /* (73) orconf ::= */
+   229,  /* (74) orconf ::= OR resolvetype */
+   230,  /* (75) resolvetype ::= IGNORE */
+   230,  /* (76) resolvetype ::= REPLACE */
+   186,  /* (77) cmd ::= DROP TABLE ifexists fullname */
+   232,  /* (78) ifexists ::= IF EXISTS */
+   232,  /* (79) ifexists ::= */
+   186,  /* (80) cmd ::= createkw temp VIEW ifnotexists nm dbnm eidlist_opt AS select */
+   186,  /* (81) cmd ::= DROP VIEW ifexists fullname */
+   186,  /* (82) cmd ::= select */
+   200,  /* (83) select ::= WITH wqlist selectnowith */
+   200,  /* (84) select ::= WITH RECURSIVE wqlist selectnowith */
+   200,  /* (85) select ::= selectnowith */
+   234,  /* (86) selectnowith ::= selectnowith multiselect_op oneselect */
+   237,  /* (87) multiselect_op ::= UNION */
+   237,  /* (88) multiselect_op ::= UNION ALL */
+   237,  /* (89) multiselect_op ::= EXCEPT|INTERSECT */
+   235,  /* (90) oneselect ::= SELECT distinct selcollist from where_opt groupby_opt having_opt orderby_opt limit_opt */
+   235,  /* (91) oneselect ::= SELECT distinct selcollist from where_opt groupby_opt having_opt window_clause orderby_opt limit_opt */
+   247,  /* (92) values ::= VALUES LP nexprlist RP */
+   247,  /* (93) values ::= values COMMA LP nexprlist RP */
+   238,  /* (94) distinct ::= DISTINCT */
+   238,  /* (95) distinct ::= ALL */
+   238,  /* (96) distinct ::= */
+   249,  /* (97) sclp ::= */
+   239,  /* (98) selcollist ::= sclp scanpt expr scanpt as */
+   239,  /* (99) selcollist ::= sclp scanpt STAR */
+   239,  /* (100) selcollist ::= sclp scanpt nm DOT STAR */
+   250,  /* (101) as ::= AS nm */
+   250,  /* (102) as ::= */
+   240,  /* (103) from ::= */
+   240,  /* (104) from ::= FROM seltablist */
+   252,  /* (105) stl_prefix ::= seltablist joinop */
+   252,  /* (106) stl_prefix ::= */
+   251,  /* (107) seltablist ::= stl_prefix nm dbnm as indexed_opt on_opt using_opt */
+   251,  /* (108) seltablist ::= stl_prefix nm dbnm LP exprlist RP as on_opt using_opt */
+   251,  /* (109) seltablist ::= stl_prefix LP select RP as on_opt using_opt */
+   251,  /* (110) seltablist ::= stl_prefix LP seltablist RP as on_opt using_opt */
+   196,  /* (111) dbnm ::= */
+   196,  /* (112) dbnm ::= DOT nm */
+   233,  /* (113) fullname ::= nm */
+   233,  /* (114) fullname ::= nm DOT nm */
+   258,  /* (115) xfullname ::= nm */
+   258,  /* (116) xfullname ::= nm DOT nm */
+   258,  /* (117) xfullname ::= nm DOT nm AS nm */
+   258,  /* (118) xfullname ::= nm AS nm */
+   253,  /* (119) joinop ::= COMMA|JOIN */
+   253,  /* (120) joinop ::= JOIN_KW JOIN */
+   253,  /* (121) joinop ::= JOIN_KW nm JOIN */
+   253,  /* (122) joinop ::= JOIN_KW nm nm JOIN */
+   255,  /* (123) on_opt ::= ON expr */
+   255,  /* (124) on_opt ::= */
+   254,  /* (125) indexed_opt ::= */
+   254,  /* (126) indexed_opt ::= INDEXED BY nm */
+   254,  /* (127) indexed_opt ::= NOT INDEXED */
+   256,  /* (128) using_opt ::= USING LP idlist RP */
+   256,  /* (129) using_opt ::= */
+   244,  /* (130) orderby_opt ::= */
+   244,  /* (131) orderby_opt ::= ORDER BY sortlist */
+   226,  /* (132) sortlist ::= sortlist COMMA expr sortorder nulls */
+   226,  /* (133) sortlist ::= expr sortorder nulls */
+   214,  /* (134) sortorder ::= ASC */
+   214,  /* (135) sortorder ::= DESC */
+   214,  /* (136) sortorder ::= */
+   260,  /* (137) nulls ::= NULLS FIRST */
+   260,  /* (138) nulls ::= NULLS LAST */
+   260,  /* (139) nulls ::= */
+   242,  /* (140) groupby_opt ::= */
+   242,  /* (141) groupby_opt ::= GROUP BY nexprlist */
+   243,  /* (142) having_opt ::= */
+   243,  /* (143) having_opt ::= HAVING expr */
+   245,  /* (144) limit_opt ::= */
+   245,  /* (145) limit_opt ::= LIMIT expr */
+   245,  /* (146) limit_opt ::= LIMIT expr OFFSET expr */
+   245,  /* (147) limit_opt ::= LIMIT expr COMMA expr */
+   186,  /* (148) cmd ::= with DELETE FROM xfullname indexed_opt where_opt */
+   241,  /* (149) where_opt ::= */
+   241,  /* (150) where_opt ::= WHERE expr */
+   186,  /* (151) cmd ::= with UPDATE orconf xfullname indexed_opt SET setlist where_opt */
+   262,  /* (152) setlist ::= setlist COMMA nm EQ expr */
+   262,  /* (153) setlist ::= setlist COMMA LP idlist RP EQ expr */
+   262,  /* (154) setlist ::= nm EQ expr */
+   262,  /* (155) setlist ::= LP idlist RP EQ expr */
+   186,  /* (156) cmd ::= with insert_cmd INTO xfullname idlist_opt select upsert */
+   186,  /* (157) cmd ::= with insert_cmd INTO xfullname idlist_opt DEFAULT VALUES */
+   265,  /* (158) upsert ::= */
+   265,  /* (159) upsert ::= ON CONFLICT LP sortlist RP where_opt DO UPDATE SET setlist where_opt */
+   265,  /* (160) upsert ::= ON CONFLICT LP sortlist RP where_opt DO NOTHING */
+   265,  /* (161) upsert ::= ON CONFLICT DO NOTHING */
+   263,  /* (162) insert_cmd ::= INSERT orconf */
+   263,  /* (163) insert_cmd ::= REPLACE */
+   264,  /* (164) idlist_opt ::= */
+   264,  /* (165) idlist_opt ::= LP idlist RP */
+   259,  /* (166) idlist ::= idlist COMMA nm */
+   259,  /* (167) idlist ::= nm */
+   212,  /* (168) expr ::= LP expr RP */
+   212,  /* (169) expr ::= ID|INDEXED */
+   212,  /* (170) expr ::= JOIN_KW */
+   212,  /* (171) expr ::= nm DOT nm */
+   212,  /* (172) expr ::= nm DOT nm DOT nm */
+   211,  /* (173) term ::= NULL|FLOAT|BLOB */
+   211,  /* (174) term ::= STRING */
+   211,  /* (175) term ::= INTEGER */
+   212,  /* (176) expr ::= VARIABLE */
+   212,  /* (177) expr ::= expr COLLATE ID|STRING */
+   212,  /* (178) expr ::= CAST LP expr AS typetoken RP */
+   212,  /* (179) expr ::= ID|INDEXED LP distinct exprlist RP */
+   212,  /* (180) expr ::= ID|INDEXED LP STAR RP */
+   212,  /* (181) expr ::= ID|INDEXED LP distinct exprlist RP filter_over */
+   212,  /* (182) expr ::= ID|INDEXED LP STAR RP filter_over */
+   211,  /* (183) term ::= CTIME_KW */
+   212,  /* (184) expr ::= LP nexprlist COMMA expr RP */
+   212,  /* (185) expr ::= expr AND expr */
+   212,  /* (186) expr ::= expr OR expr */
+   212,  /* (187) expr ::= expr LT|GT|GE|LE expr */
+   212,  /* (188) expr ::= expr EQ|NE expr */
+   212,  /* (189) expr ::= expr BITAND|BITOR|LSHIFT|RSHIFT expr */
+   212,  /* (190) expr ::= expr PLUS|MINUS expr */
+   212,  /* (191) expr ::= expr STAR|SLASH|REM expr */
+   212,  /* (192) expr ::= expr CONCAT expr */
+   267,  /* (193) likeop ::= NOT LIKE_KW|MATCH */
+   212,  /* (194) expr ::= expr likeop expr */
+   212,  /* (195) expr ::= expr likeop expr ESCAPE expr */
+   212,  /* (196) expr ::= expr ISNULL|NOTNULL */
+   212,  /* (197) expr ::= expr NOT NULL */
+   212,  /* (198) expr ::= expr IS expr */
+   212,  /* (199) expr ::= expr IS NOT expr */
+   212,  /* (200) expr ::= NOT expr */
+   212,  /* (201) expr ::= BITNOT expr */
+   212,  /* (202) expr ::= PLUS|MINUS expr */
+   268,  /* (203) between_op ::= BETWEEN */
+   268,  /* (204) between_op ::= NOT BETWEEN */
+   212,  /* (205) expr ::= expr between_op expr AND expr */
+   269,  /* (206) in_op ::= IN */
+   269,  /* (207) in_op ::= NOT IN */
+   212,  /* (208) expr ::= expr in_op LP exprlist RP */
+   212,  /* (209) expr ::= LP select RP */
+   212,  /* (210) expr ::= expr in_op LP select RP */
+   212,  /* (211) expr ::= expr in_op nm dbnm paren_exprlist */
+   212,  /* (212) expr ::= EXISTS LP select RP */
+   212,  /* (213) expr ::= CASE case_operand case_exprlist case_else END */
+   272,  /* (214) case_exprlist ::= case_exprlist WHEN expr THEN expr */
+   272,  /* (215) case_exprlist ::= WHEN expr THEN expr */
+   273,  /* (216) case_else ::= ELSE expr */
+   273,  /* (217) case_else ::= */
+   271,  /* (218) case_operand ::= expr */
+   271,  /* (219) case_operand ::= */
+   257,  /* (220) exprlist ::= */
+   248,  /* (221) nexprlist ::= nexprlist COMMA expr */
+   248,  /* (222) nexprlist ::= expr */
+   270,  /* (223) paren_exprlist ::= */
+   270,  /* (224) paren_exprlist ::= LP exprlist RP */
+   186,  /* (225) cmd ::= createkw uniqueflag INDEX ifnotexists nm dbnm ON nm LP sortlist RP where_opt */
+   274,  /* (226) uniqueflag ::= UNIQUE */
+   274,  /* (227) uniqueflag ::= */
+   216,  /* (228) eidlist_opt ::= */
+   216,  /* (229) eidlist_opt ::= LP eidlist RP */
+   227,  /* (230) eidlist ::= eidlist COMMA nm collate sortorder */
+   227,  /* (231) eidlist ::= nm collate sortorder */
+   275,  /* (232) collate ::= */
+   275,  /* (233) collate ::= COLLATE ID|STRING */
+   186,  /* (234) cmd ::= DROP INDEX ifexists fullname */
+   186,  /* (235) cmd ::= VACUUM vinto */
+   186,  /* (236) cmd ::= VACUUM nm vinto */
+   276,  /* (237) vinto ::= INTO expr */
+   276,  /* (238) vinto ::= */
+   186,  /* (239) cmd ::= PRAGMA nm dbnm */
+   186,  /* (240) cmd ::= PRAGMA nm dbnm EQ nmnum */
+   186,  /* (241) cmd ::= PRAGMA nm dbnm LP nmnum RP */
+   186,  /* (242) cmd ::= PRAGMA nm dbnm EQ minus_num */
+   186,  /* (243) cmd ::= PRAGMA nm dbnm LP minus_num RP */
+   206,  /* (244) plus_num ::= PLUS INTEGER|FLOAT */
+   207,  /* (245) minus_num ::= MINUS INTEGER|FLOAT */
+   186,  /* (246) cmd ::= createkw trigger_decl BEGIN trigger_cmd_list END */
+   278,  /* (247) trigger_decl ::= temp TRIGGER ifnotexists nm dbnm trigger_time trigger_event ON fullname foreach_clause when_clause */
+   280,  /* (248) trigger_time ::= BEFORE|AFTER */
+   280,  /* (249) trigger_time ::= INSTEAD OF */
+   280,  /* (250) trigger_time ::= */
+   281,  /* (251) trigger_event ::= DELETE|INSERT */
+   281,  /* (252) trigger_event ::= UPDATE */
+   281,  /* (253) trigger_event ::= UPDATE OF idlist */
+   283,  /* (254) when_clause ::= */
+   283,  /* (255) when_clause ::= WHEN expr */
+   279,  /* (256) trigger_cmd_list ::= trigger_cmd_list trigger_cmd SEMI */
+   279,  /* (257) trigger_cmd_list ::= trigger_cmd SEMI */
+   285,  /* (258) trnm ::= nm DOT nm */
+   286,  /* (259) tridxby ::= INDEXED BY nm */
+   286,  /* (260) tridxby ::= NOT INDEXED */
+   284,  /* (261) trigger_cmd ::= UPDATE orconf trnm tridxby SET setlist where_opt scanpt */
+   284,  /* (262) trigger_cmd ::= scanpt insert_cmd INTO trnm idlist_opt select upsert scanpt */
+   284,  /* (263) trigger_cmd ::= DELETE FROM trnm tridxby where_opt scanpt */
+   284,  /* (264) trigger_cmd ::= scanpt select scanpt */
+   212,  /* (265) expr ::= RAISE LP IGNORE RP */
+   212,  /* (266) expr ::= RAISE LP raisetype COMMA nm RP */
+   231,  /* (267) raisetype ::= ROLLBACK */
+   231,  /* (268) raisetype ::= ABORT */
+   231,  /* (269) raisetype ::= FAIL */
+   186,  /* (270) cmd ::= DROP TRIGGER ifexists fullname */
+   186,  /* (271) cmd ::= ATTACH database_kw_opt expr AS expr key_opt */
+   186,  /* (272) cmd ::= DETACH database_kw_opt expr */
+   288,  /* (273) key_opt ::= */
+   288,  /* (274) key_opt ::= KEY expr */
+   186,  /* (275) cmd ::= REINDEX */
+   186,  /* (276) cmd ::= REINDEX nm dbnm */
+   186,  /* (277) cmd ::= ANALYZE */
+   186,  /* (278) cmd ::= ANALYZE nm dbnm */
+   186,  /* (279) cmd ::= ALTER TABLE fullname RENAME TO nm */
+   186,  /* (280) cmd ::= ALTER TABLE add_column_fullname ADD kwcolumn_opt columnname carglist */
+   289,  /* (281) add_column_fullname ::= fullname */
+   186,  /* (282) cmd ::= ALTER TABLE fullname RENAME kwcolumn_opt nm TO nm */
+   186,  /* (283) cmd ::= create_vtab */
+   186,  /* (284) cmd ::= create_vtab LP vtabarglist RP */
+   291,  /* (285) create_vtab ::= createkw VIRTUAL TABLE ifnotexists nm dbnm USING nm */
+   293,  /* (286) vtabarg ::= */
+   294,  /* (287) vtabargtoken ::= ANY */
+   294,  /* (288) vtabargtoken ::= lp anylist RP */
+   295,  /* (289) lp ::= LP */
+   261,  /* (290) with ::= WITH wqlist */
+   261,  /* (291) with ::= WITH RECURSIVE wqlist */
+   236,  /* (292) wqlist ::= nm eidlist_opt AS LP select RP */
+   236,  /* (293) wqlist ::= wqlist COMMA nm eidlist_opt AS LP select RP */
+   297,  /* (294) windowdefn_list ::= windowdefn */
+   297,  /* (295) windowdefn_list ::= windowdefn_list COMMA windowdefn */
+   298,  /* (296) windowdefn ::= nm AS LP window RP */
+   299,  /* (297) window ::= PARTITION BY nexprlist orderby_opt frame_opt */
+   299,  /* (298) window ::= nm PARTITION BY nexprlist orderby_opt frame_opt */
+   299,  /* (299) window ::= ORDER BY sortlist frame_opt */
+   299,  /* (300) window ::= nm ORDER BY sortlist frame_opt */
+   299,  /* (301) window ::= frame_opt */
+   299,  /* (302) window ::= nm frame_opt */
+   300,  /* (303) frame_opt ::= */
+   300,  /* (304) frame_opt ::= range_or_rows frame_bound_s frame_exclude_opt */
+   300,  /* (305) frame_opt ::= range_or_rows BETWEEN frame_bound_s AND frame_bound_e frame_exclude_opt */
+   304,  /* (306) range_or_rows ::= RANGE|ROWS|GROUPS */
+   306,  /* (307) frame_bound_s ::= frame_bound */
+   306,  /* (308) frame_bound_s ::= UNBOUNDED PRECEDING */
+   307,  /* (309) frame_bound_e ::= frame_bound */
+   307,  /* (310) frame_bound_e ::= UNBOUNDED FOLLOWING */
+   305,  /* (311) frame_bound ::= expr PRECEDING|FOLLOWING */
+   305,  /* (312) frame_bound ::= CURRENT ROW */
+   308,  /* (313) frame_exclude_opt ::= */
+   308,  /* (314) frame_exclude_opt ::= EXCLUDE frame_exclude */
+   309,  /* (315) frame_exclude ::= NO OTHERS */
+   309,  /* (316) frame_exclude ::= CURRENT ROW */
+   309,  /* (317) frame_exclude ::= GROUP|TIES */
+   246,  /* (318) window_clause ::= WINDOW windowdefn_list */
+   266,  /* (319) filter_over ::= filter_clause over_clause */
+   266,  /* (320) filter_over ::= over_clause */
+   266,  /* (321) filter_over ::= filter_clause */
+   303,  /* (322) over_clause ::= OVER LP window RP */
+   303,  /* (323) over_clause ::= OVER nm */
+   302,  /* (324) filter_clause ::= FILTER LP WHERE expr RP */
+   181,  /* (325) input ::= cmdlist */
+   182,  /* (326) cmdlist ::= cmdlist ecmd */
+   182,  /* (327) cmdlist ::= ecmd */
+   183,  /* (328) ecmd ::= SEMI */
+   183,  /* (329) ecmd ::= cmdx SEMI */
+   183,  /* (330) ecmd ::= explain cmdx SEMI */
+   188,  /* (331) trans_opt ::= */
+   188,  /* (332) trans_opt ::= TRANSACTION */
+   188,  /* (333) trans_opt ::= TRANSACTION nm */
+   190,  /* (334) savepoint_opt ::= SAVEPOINT */
+   190,  /* (335) savepoint_opt ::= */
+   186,  /* (336) cmd ::= create_table create_table_args */
+   197,  /* (337) columnlist ::= columnlist COMMA columnname carglist */
+   197,  /* (338) columnlist ::= columnname carglist */
+   189,  /* (339) nm ::= ID|INDEXED */
+   189,  /* (340) nm ::= STRING */
+   189,  /* (341) nm ::= JOIN_KW */
+   203,  /* (342) typetoken ::= typename */
+   204,  /* (343) typename ::= ID|STRING */
+   205,  /* (344) signed ::= plus_num */
+   205,  /* (345) signed ::= minus_num */
+   202,  /* (346) carglist ::= carglist ccons */
+   202,  /* (347) carglist ::= */
+   210,  /* (348) ccons ::= NULL onconf */
+   210,  /* (349) ccons ::= GENERATED ALWAYS AS generated */
+   210,  /* (350) ccons ::= AS generated */
+   198,  /* (351) conslist_opt ::= COMMA conslist */
+   223,  /* (352) conslist ::= conslist tconscomma tcons */
+   223,  /* (353) conslist ::= tcons */
+   224,  /* (354) tconscomma ::= */
+   228,  /* (355) defer_subclause_opt ::= defer_subclause */
+   230,  /* (356) resolvetype ::= raisetype */
+   234,  /* (357) selectnowith ::= oneselect */
+   235,  /* (358) oneselect ::= values */
+   249,  /* (359) sclp ::= selcollist COMMA */
+   250,  /* (360) as ::= ID|STRING */
+   212,  /* (361) expr ::= term */
+   267,  /* (362) likeop ::= LIKE_KW|MATCH */
+   257,  /* (363) exprlist ::= nexprlist */
+   277,  /* (364) nmnum ::= plus_num */
+   277,  /* (365) nmnum ::= nm */
+   277,  /* (366) nmnum ::= ON */
+   277,  /* (367) nmnum ::= DELETE */
+   277,  /* (368) nmnum ::= DEFAULT */
+   206,  /* (369) plus_num ::= INTEGER|FLOAT */
+   282,  /* (370) foreach_clause ::= */
+   282,  /* (371) foreach_clause ::= FOR EACH ROW */
+   285,  /* (372) trnm ::= nm */
+   286,  /* (373) tridxby ::= */
+   287,  /* (374) database_kw_opt ::= DATABASE */
+   287,  /* (375) database_kw_opt ::= */
+   290,  /* (376) kwcolumn_opt ::= */
+   290,  /* (377) kwcolumn_opt ::= COLUMNKW */
+   292,  /* (378) vtabarglist ::= vtabarg */
+   292,  /* (379) vtabarglist ::= vtabarglist COMMA vtabarg */
+   293,  /* (380) vtabarg ::= vtabarg vtabargtoken */
+   296,  /* (381) anylist ::= */
+   296,  /* (382) anylist ::= anylist LP anylist RP */
+   296,  /* (383) anylist ::= anylist ANY */
+   261,  /* (384) with ::= */
 };
 
 /* For rule J, yyRuleInfoNRhs[J] contains the negative of the number
@@ -150879,352 +155160,362 @@ static const signed char yyRuleInfoNRhs[] = {
    -6,  /* (26) typetoken ::= typename LP signed COMMA signed RP */
    -2,  /* (27) typename ::= typename ID|STRING */
     0,  /* (28) scanpt ::= */
-   -2,  /* (29) ccons ::= CONSTRAINT nm */
-   -4,  /* (30) ccons ::= DEFAULT scanpt term scanpt */
-   -4,  /* (31) ccons ::= DEFAULT LP expr RP */
-   -4,  /* (32) ccons ::= DEFAULT PLUS term scanpt */
-   -4,  /* (33) ccons ::= DEFAULT MINUS term scanpt */
-   -3,  /* (34) ccons ::= DEFAULT scanpt ID|INDEXED */
-   -3,  /* (35) ccons ::= NOT NULL onconf */
-   -5,  /* (36) ccons ::= PRIMARY KEY sortorder onconf autoinc */
-   -2,  /* (37) ccons ::= UNIQUE onconf */
-   -4,  /* (38) ccons ::= CHECK LP expr RP */
-   -4,  /* (39) ccons ::= REFERENCES nm eidlist_opt refargs */
-   -1,  /* (40) ccons ::= defer_subclause */
-   -2,  /* (41) ccons ::= COLLATE ID|STRING */
-    0,  /* (42) autoinc ::= */
-   -1,  /* (43) autoinc ::= AUTOINCR */
-    0,  /* (44) refargs ::= */
-   -2,  /* (45) refargs ::= refargs refarg */
-   -2,  /* (46) refarg ::= MATCH nm */
-   -3,  /* (47) refarg ::= ON INSERT refact */
-   -3,  /* (48) refarg ::= ON DELETE refact */
-   -3,  /* (49) refarg ::= ON UPDATE refact */
-   -2,  /* (50) refact ::= SET NULL */
-   -2,  /* (51) refact ::= SET DEFAULT */
-   -1,  /* (52) refact ::= CASCADE */
-   -1,  /* (53) refact ::= RESTRICT */
-   -2,  /* (54) refact ::= NO ACTION */
-   -3,  /* (55) defer_subclause ::= NOT DEFERRABLE init_deferred_pred_opt */
-   -2,  /* (56) defer_subclause ::= DEFERRABLE init_deferred_pred_opt */
-    0,  /* (57) init_deferred_pred_opt ::= */
-   -2,  /* (58) init_deferred_pred_opt ::= INITIALLY DEFERRED */
-   -2,  /* (59) init_deferred_pred_opt ::= INITIALLY IMMEDIATE */
-    0,  /* (60) conslist_opt ::= */
-   -1,  /* (61) tconscomma ::= COMMA */
-   -2,  /* (62) tcons ::= CONSTRAINT nm */
-   -7,  /* (63) tcons ::= PRIMARY KEY LP sortlist autoinc RP onconf */
-   -5,  /* (64) tcons ::= UNIQUE LP sortlist RP onconf */
-   -5,  /* (65) tcons ::= CHECK LP expr RP onconf */
-  -10,  /* (66) tcons ::= FOREIGN KEY LP eidlist RP REFERENCES nm eidlist_opt refargs defer_subclause_opt */
-    0,  /* (67) defer_subclause_opt ::= */
-    0,  /* (68) onconf ::= */
-   -3,  /* (69) onconf ::= ON CONFLICT resolvetype */
-    0,  /* (70) orconf ::= */
-   -2,  /* (71) orconf ::= OR resolvetype */
-   -1,  /* (72) resolvetype ::= IGNORE */
-   -1,  /* (73) resolvetype ::= REPLACE */
-   -4,  /* (74) cmd ::= DROP TABLE ifexists fullname */
-   -2,  /* (75) ifexists ::= IF EXISTS */
-    0,  /* (76) ifexists ::= */
-   -9,  /* (77) cmd ::= createkw temp VIEW ifnotexists nm dbnm eidlist_opt AS select */
-   -4,  /* (78) cmd ::= DROP VIEW ifexists fullname */
-   -1,  /* (79) cmd ::= select */
-   -3,  /* (80) select ::= WITH wqlist selectnowith */
-   -4,  /* (81) select ::= WITH RECURSIVE wqlist selectnowith */
-   -1,  /* (82) select ::= selectnowith */
-   -3,  /* (83) selectnowith ::= selectnowith multiselect_op oneselect */
-   -1,  /* (84) multiselect_op ::= UNION */
-   -2,  /* (85) multiselect_op ::= UNION ALL */
-   -1,  /* (86) multiselect_op ::= EXCEPT|INTERSECT */
-   -9,  /* (87) oneselect ::= SELECT distinct selcollist from where_opt groupby_opt having_opt orderby_opt limit_opt */
-  -10,  /* (88) oneselect ::= SELECT distinct selcollist from where_opt groupby_opt having_opt window_clause orderby_opt limit_opt */
-   -4,  /* (89) values ::= VALUES LP nexprlist RP */
-   -5,  /* (90) values ::= values COMMA LP nexprlist RP */
-   -1,  /* (91) distinct ::= DISTINCT */
-   -1,  /* (92) distinct ::= ALL */
-    0,  /* (93) distinct ::= */
-    0,  /* (94) sclp ::= */
-   -5,  /* (95) selcollist ::= sclp scanpt expr scanpt as */
-   -3,  /* (96) selcollist ::= sclp scanpt STAR */
-   -5,  /* (97) selcollist ::= sclp scanpt nm DOT STAR */
-   -2,  /* (98) as ::= AS nm */
-    0,  /* (99) as ::= */
-    0,  /* (100) from ::= */
-   -2,  /* (101) from ::= FROM seltablist */
-   -2,  /* (102) stl_prefix ::= seltablist joinop */
-    0,  /* (103) stl_prefix ::= */
-   -7,  /* (104) seltablist ::= stl_prefix nm dbnm as indexed_opt on_opt using_opt */
-   -9,  /* (105) seltablist ::= stl_prefix nm dbnm LP exprlist RP as on_opt using_opt */
-   -7,  /* (106) seltablist ::= stl_prefix LP select RP as on_opt using_opt */
-   -7,  /* (107) seltablist ::= stl_prefix LP seltablist RP as on_opt using_opt */
-    0,  /* (108) dbnm ::= */
-   -2,  /* (109) dbnm ::= DOT nm */
-   -1,  /* (110) fullname ::= nm */
-   -3,  /* (111) fullname ::= nm DOT nm */
-   -1,  /* (112) xfullname ::= nm */
-   -3,  /* (113) xfullname ::= nm DOT nm */
-   -5,  /* (114) xfullname ::= nm DOT nm AS nm */
-   -3,  /* (115) xfullname ::= nm AS nm */
-   -1,  /* (116) joinop ::= COMMA|JOIN */
-   -2,  /* (117) joinop ::= JOIN_KW JOIN */
-   -3,  /* (118) joinop ::= JOIN_KW nm JOIN */
-   -4,  /* (119) joinop ::= JOIN_KW nm nm JOIN */
-   -2,  /* (120) on_opt ::= ON expr */
-    0,  /* (121) on_opt ::= */
-    0,  /* (122) indexed_opt ::= */
-   -3,  /* (123) indexed_opt ::= INDEXED BY nm */
-   -2,  /* (124) indexed_opt ::= NOT INDEXED */
-   -4,  /* (125) using_opt ::= USING LP idlist RP */
-    0,  /* (126) using_opt ::= */
-    0,  /* (127) orderby_opt ::= */
-   -3,  /* (128) orderby_opt ::= ORDER BY sortlist */
-   -4,  /* (129) sortlist ::= sortlist COMMA expr sortorder */
-   -2,  /* (130) sortlist ::= expr sortorder */
-   -1,  /* (131) sortorder ::= ASC */
-   -1,  /* (132) sortorder ::= DESC */
-    0,  /* (133) sortorder ::= */
-    0,  /* (134) groupby_opt ::= */
-   -3,  /* (135) groupby_opt ::= GROUP BY nexprlist */
-    0,  /* (136) having_opt ::= */
-   -2,  /* (137) having_opt ::= HAVING expr */
-    0,  /* (138) limit_opt ::= */
-   -2,  /* (139) limit_opt ::= LIMIT expr */
-   -4,  /* (140) limit_opt ::= LIMIT expr OFFSET expr */
-   -4,  /* (141) limit_opt ::= LIMIT expr COMMA expr */
-   -6,  /* (142) cmd ::= with DELETE FROM xfullname indexed_opt where_opt */
-    0,  /* (143) where_opt ::= */
-   -2,  /* (144) where_opt ::= WHERE expr */
-   -8,  /* (145) cmd ::= with UPDATE orconf xfullname indexed_opt SET setlist where_opt */
-   -5,  /* (146) setlist ::= setlist COMMA nm EQ expr */
-   -7,  /* (147) setlist ::= setlist COMMA LP idlist RP EQ expr */
-   -3,  /* (148) setlist ::= nm EQ expr */
-   -5,  /* (149) setlist ::= LP idlist RP EQ expr */
-   -7,  /* (150) cmd ::= with insert_cmd INTO xfullname idlist_opt select upsert */
-   -7,  /* (151) cmd ::= with insert_cmd INTO xfullname idlist_opt DEFAULT VALUES */
-    0,  /* (152) upsert ::= */
-  -11,  /* (153) upsert ::= ON CONFLICT LP sortlist RP where_opt DO UPDATE SET setlist where_opt */
-   -8,  /* (154) upsert ::= ON CONFLICT LP sortlist RP where_opt DO NOTHING */
-   -4,  /* (155) upsert ::= ON CONFLICT DO NOTHING */
-   -2,  /* (156) insert_cmd ::= INSERT orconf */
-   -1,  /* (157) insert_cmd ::= REPLACE */
-    0,  /* (158) idlist_opt ::= */
-   -3,  /* (159) idlist_opt ::= LP idlist RP */
-   -3,  /* (160) idlist ::= idlist COMMA nm */
-   -1,  /* (161) idlist ::= nm */
-   -3,  /* (162) expr ::= LP expr RP */
-   -1,  /* (163) expr ::= ID|INDEXED */
-   -1,  /* (164) expr ::= JOIN_KW */
-   -3,  /* (165) expr ::= nm DOT nm */
-   -5,  /* (166) expr ::= nm DOT nm DOT nm */
-   -1,  /* (167) term ::= NULL|FLOAT|BLOB */
-   -1,  /* (168) term ::= STRING */
-   -1,  /* (169) term ::= INTEGER */
-   -1,  /* (170) expr ::= VARIABLE */
-   -3,  /* (171) expr ::= expr COLLATE ID|STRING */
-   -6,  /* (172) expr ::= CAST LP expr AS typetoken RP */
-   -5,  /* (173) expr ::= ID|INDEXED LP distinct exprlist RP */
-   -4,  /* (174) expr ::= ID|INDEXED LP STAR RP */
-   -6,  /* (175) expr ::= ID|INDEXED LP distinct exprlist RP over_clause */
-   -5,  /* (176) expr ::= ID|INDEXED LP STAR RP over_clause */
-   -1,  /* (177) term ::= CTIME_KW */
-   -5,  /* (178) expr ::= LP nexprlist COMMA expr RP */
-   -3,  /* (179) expr ::= expr AND expr */
-   -3,  /* (180) expr ::= expr OR expr */
-   -3,  /* (181) expr ::= expr LT|GT|GE|LE expr */
-   -3,  /* (182) expr ::= expr EQ|NE expr */
-   -3,  /* (183) expr ::= expr BITAND|BITOR|LSHIFT|RSHIFT expr */
-   -3,  /* (184) expr ::= expr PLUS|MINUS expr */
-   -3,  /* (185) expr ::= expr STAR|SLASH|REM expr */
-   -3,  /* (186) expr ::= expr CONCAT expr */
-   -2,  /* (187) likeop ::= NOT LIKE_KW|MATCH */
-   -3,  /* (188) expr ::= expr likeop expr */
-   -5,  /* (189) expr ::= expr likeop expr ESCAPE expr */
-   -2,  /* (190) expr ::= expr ISNULL|NOTNULL */
-   -3,  /* (191) expr ::= expr NOT NULL */
-   -3,  /* (192) expr ::= expr IS expr */
-   -4,  /* (193) expr ::= expr IS NOT expr */
-   -2,  /* (194) expr ::= NOT expr */
-   -2,  /* (195) expr ::= BITNOT expr */
-   -2,  /* (196) expr ::= PLUS|MINUS expr */
-   -1,  /* (197) between_op ::= BETWEEN */
-   -2,  /* (198) between_op ::= NOT BETWEEN */
-   -5,  /* (199) expr ::= expr between_op expr AND expr */
-   -1,  /* (200) in_op ::= IN */
-   -2,  /* (201) in_op ::= NOT IN */
-   -5,  /* (202) expr ::= expr in_op LP exprlist RP */
-   -3,  /* (203) expr ::= LP select RP */
-   -5,  /* (204) expr ::= expr in_op LP select RP */
-   -5,  /* (205) expr ::= expr in_op nm dbnm paren_exprlist */
-   -4,  /* (206) expr ::= EXISTS LP select RP */
-   -5,  /* (207) expr ::= CASE case_operand case_exprlist case_else END */
-   -5,  /* (208) case_exprlist ::= case_exprlist WHEN expr THEN expr */
-   -4,  /* (209) case_exprlist ::= WHEN expr THEN expr */
-   -2,  /* (210) case_else ::= ELSE expr */
-    0,  /* (211) case_else ::= */
-   -1,  /* (212) case_operand ::= expr */
-    0,  /* (213) case_operand ::= */
-    0,  /* (214) exprlist ::= */
-   -3,  /* (215) nexprlist ::= nexprlist COMMA expr */
-   -1,  /* (216) nexprlist ::= expr */
-    0,  /* (217) paren_exprlist ::= */
-   -3,  /* (218) paren_exprlist ::= LP exprlist RP */
-  -12,  /* (219) cmd ::= createkw uniqueflag INDEX ifnotexists nm dbnm ON nm LP sortlist RP where_opt */
-   -1,  /* (220) uniqueflag ::= UNIQUE */
-    0,  /* (221) uniqueflag ::= */
-    0,  /* (222) eidlist_opt ::= */
-   -3,  /* (223) eidlist_opt ::= LP eidlist RP */
-   -5,  /* (224) eidlist ::= eidlist COMMA nm collate sortorder */
-   -3,  /* (225) eidlist ::= nm collate sortorder */
-    0,  /* (226) collate ::= */
-   -2,  /* (227) collate ::= COLLATE ID|STRING */
-   -4,  /* (228) cmd ::= DROP INDEX ifexists fullname */
-   -2,  /* (229) cmd ::= VACUUM vinto */
-   -3,  /* (230) cmd ::= VACUUM nm vinto */
-   -2,  /* (231) vinto ::= INTO expr */
-    0,  /* (232) vinto ::= */
-   -3,  /* (233) cmd ::= PRAGMA nm dbnm */
-   -5,  /* (234) cmd ::= PRAGMA nm dbnm EQ nmnum */
-   -6,  /* (235) cmd ::= PRAGMA nm dbnm LP nmnum RP */
-   -5,  /* (236) cmd ::= PRAGMA nm dbnm EQ minus_num */
-   -6,  /* (237) cmd ::= PRAGMA nm dbnm LP minus_num RP */
-   -2,  /* (238) plus_num ::= PLUS INTEGER|FLOAT */
-   -2,  /* (239) minus_num ::= MINUS INTEGER|FLOAT */
-   -5,  /* (240) cmd ::= createkw trigger_decl BEGIN trigger_cmd_list END */
-  -11,  /* (241) trigger_decl ::= temp TRIGGER ifnotexists nm dbnm trigger_time trigger_event ON fullname foreach_clause when_clause */
-   -1,  /* (242) trigger_time ::= BEFORE|AFTER */
-   -2,  /* (243) trigger_time ::= INSTEAD OF */
-    0,  /* (244) trigger_time ::= */
-   -1,  /* (245) trigger_event ::= DELETE|INSERT */
-   -1,  /* (246) trigger_event ::= UPDATE */
-   -3,  /* (247) trigger_event ::= UPDATE OF idlist */
-    0,  /* (248) when_clause ::= */
-   -2,  /* (249) when_clause ::= WHEN expr */
-   -3,  /* (250) trigger_cmd_list ::= trigger_cmd_list trigger_cmd SEMI */
-   -2,  /* (251) trigger_cmd_list ::= trigger_cmd SEMI */
-   -3,  /* (252) trnm ::= nm DOT nm */
-   -3,  /* (253) tridxby ::= INDEXED BY nm */
-   -2,  /* (254) tridxby ::= NOT INDEXED */
-   -8,  /* (255) trigger_cmd ::= UPDATE orconf trnm tridxby SET setlist where_opt scanpt */
-   -8,  /* (256) trigger_cmd ::= scanpt insert_cmd INTO trnm idlist_opt select upsert scanpt */
-   -6,  /* (257) trigger_cmd ::= DELETE FROM trnm tridxby where_opt scanpt */
-   -3,  /* (258) trigger_cmd ::= scanpt select scanpt */
-   -4,  /* (259) expr ::= RAISE LP IGNORE RP */
-   -6,  /* (260) expr ::= RAISE LP raisetype COMMA nm RP */
-   -1,  /* (261) raisetype ::= ROLLBACK */
-   -1,  /* (262) raisetype ::= ABORT */
-   -1,  /* (263) raisetype ::= FAIL */
-   -4,  /* (264) cmd ::= DROP TRIGGER ifexists fullname */
-   -6,  /* (265) cmd ::= ATTACH database_kw_opt expr AS expr key_opt */
-   -3,  /* (266) cmd ::= DETACH database_kw_opt expr */
-    0,  /* (267) key_opt ::= */
-   -2,  /* (268) key_opt ::= KEY expr */
-   -1,  /* (269) cmd ::= REINDEX */
-   -3,  /* (270) cmd ::= REINDEX nm dbnm */
-   -1,  /* (271) cmd ::= ANALYZE */
-   -3,  /* (272) cmd ::= ANALYZE nm dbnm */
-   -6,  /* (273) cmd ::= ALTER TABLE fullname RENAME TO nm */
-   -7,  /* (274) cmd ::= ALTER TABLE add_column_fullname ADD kwcolumn_opt columnname carglist */
-   -1,  /* (275) add_column_fullname ::= fullname */
-   -8,  /* (276) cmd ::= ALTER TABLE fullname RENAME kwcolumn_opt nm TO nm */
-   -1,  /* (277) cmd ::= create_vtab */
-   -4,  /* (278) cmd ::= create_vtab LP vtabarglist RP */
-   -8,  /* (279) create_vtab ::= createkw VIRTUAL TABLE ifnotexists nm dbnm USING nm */
-    0,  /* (280) vtabarg ::= */
-   -1,  /* (281) vtabargtoken ::= ANY */
-   -3,  /* (282) vtabargtoken ::= lp anylist RP */
-   -1,  /* (283) lp ::= LP */
-   -2,  /* (284) with ::= WITH wqlist */
-   -3,  /* (285) with ::= WITH RECURSIVE wqlist */
-   -6,  /* (286) wqlist ::= nm eidlist_opt AS LP select RP */
-   -8,  /* (287) wqlist ::= wqlist COMMA nm eidlist_opt AS LP select RP */
-   -1,  /* (288) windowdefn_list ::= windowdefn */
-   -3,  /* (289) windowdefn_list ::= windowdefn_list COMMA windowdefn */
-   -5,  /* (290) windowdefn ::= nm AS LP window RP */
-   -5,  /* (291) window ::= PARTITION BY nexprlist orderby_opt frame_opt */
-   -6,  /* (292) window ::= nm PARTITION BY nexprlist orderby_opt frame_opt */
-   -4,  /* (293) window ::= ORDER BY sortlist frame_opt */
-   -5,  /* (294) window ::= nm ORDER BY sortlist frame_opt */
-   -1,  /* (295) window ::= frame_opt */
-   -2,  /* (296) window ::= nm frame_opt */
-    0,  /* (297) frame_opt ::= */
-   -3,  /* (298) frame_opt ::= range_or_rows frame_bound_s frame_exclude_opt */
-   -6,  /* (299) frame_opt ::= range_or_rows BETWEEN frame_bound_s AND frame_bound_e frame_exclude_opt */
-   -1,  /* (300) range_or_rows ::= RANGE|ROWS|GROUPS */
-   -1,  /* (301) frame_bound_s ::= frame_bound */
-   -2,  /* (302) frame_bound_s ::= UNBOUNDED PRECEDING */
-   -1,  /* (303) frame_bound_e ::= frame_bound */
-   -2,  /* (304) frame_bound_e ::= UNBOUNDED FOLLOWING */
-   -2,  /* (305) frame_bound ::= expr PRECEDING|FOLLOWING */
-   -2,  /* (306) frame_bound ::= CURRENT ROW */
-    0,  /* (307) frame_exclude_opt ::= */
-   -2,  /* (308) frame_exclude_opt ::= EXCLUDE frame_exclude */
-   -2,  /* (309) frame_exclude ::= NO OTHERS */
-   -2,  /* (310) frame_exclude ::= CURRENT ROW */
-   -1,  /* (311) frame_exclude ::= GROUP|TIES */
-   -2,  /* (312) window_clause ::= WINDOW windowdefn_list */
-   -5,  /* (313) over_clause ::= filter_opt OVER LP window RP */
-   -3,  /* (314) over_clause ::= filter_opt OVER nm */
-    0,  /* (315) filter_opt ::= */
-   -5,  /* (316) filter_opt ::= FILTER LP WHERE expr RP */
-   -1,  /* (317) input ::= cmdlist */
-   -2,  /* (318) cmdlist ::= cmdlist ecmd */
-   -1,  /* (319) cmdlist ::= ecmd */
-   -1,  /* (320) ecmd ::= SEMI */
-   -2,  /* (321) ecmd ::= cmdx SEMI */
-   -2,  /* (322) ecmd ::= explain cmdx */
-    0,  /* (323) trans_opt ::= */
-   -1,  /* (324) trans_opt ::= TRANSACTION */
-   -2,  /* (325) trans_opt ::= TRANSACTION nm */
-   -1,  /* (326) savepoint_opt ::= SAVEPOINT */
-    0,  /* (327) savepoint_opt ::= */
-   -2,  /* (328) cmd ::= create_table create_table_args */
-   -4,  /* (329) columnlist ::= columnlist COMMA columnname carglist */
-   -2,  /* (330) columnlist ::= columnname carglist */
-   -1,  /* (331) nm ::= ID|INDEXED */
-   -1,  /* (332) nm ::= STRING */
-   -1,  /* (333) nm ::= JOIN_KW */
-   -1,  /* (334) typetoken ::= typename */
-   -1,  /* (335) typename ::= ID|STRING */
-   -1,  /* (336) signed ::= plus_num */
-   -1,  /* (337) signed ::= minus_num */
-   -2,  /* (338) carglist ::= carglist ccons */
-    0,  /* (339) carglist ::= */
-   -2,  /* (340) ccons ::= NULL onconf */
-   -2,  /* (341) conslist_opt ::= COMMA conslist */
-   -3,  /* (342) conslist ::= conslist tconscomma tcons */
-   -1,  /* (343) conslist ::= tcons */
-    0,  /* (344) tconscomma ::= */
-   -1,  /* (345) defer_subclause_opt ::= defer_subclause */
-   -1,  /* (346) resolvetype ::= raisetype */
-   -1,  /* (347) selectnowith ::= oneselect */
-   -1,  /* (348) oneselect ::= values */
-   -2,  /* (349) sclp ::= selcollist COMMA */
-   -1,  /* (350) as ::= ID|STRING */
-   -1,  /* (351) expr ::= term */
-   -1,  /* (352) likeop ::= LIKE_KW|MATCH */
-   -1,  /* (353) exprlist ::= nexprlist */
-   -1,  /* (354) nmnum ::= plus_num */
-   -1,  /* (355) nmnum ::= nm */
-   -1,  /* (356) nmnum ::= ON */
-   -1,  /* (357) nmnum ::= DELETE */
-   -1,  /* (358) nmnum ::= DEFAULT */
-   -1,  /* (359) plus_num ::= INTEGER|FLOAT */
-    0,  /* (360) foreach_clause ::= */
-   -3,  /* (361) foreach_clause ::= FOR EACH ROW */
-   -1,  /* (362) trnm ::= nm */
-    0,  /* (363) tridxby ::= */
-   -1,  /* (364) database_kw_opt ::= DATABASE */
-    0,  /* (365) database_kw_opt ::= */
-    0,  /* (366) kwcolumn_opt ::= */
-   -1,  /* (367) kwcolumn_opt ::= COLUMNKW */
-   -1,  /* (368) vtabarglist ::= vtabarg */
-   -3,  /* (369) vtabarglist ::= vtabarglist COMMA vtabarg */
-   -2,  /* (370) vtabarg ::= vtabarg vtabargtoken */
-    0,  /* (371) anylist ::= */
-   -4,  /* (372) anylist ::= anylist LP anylist RP */
-   -2,  /* (373) anylist ::= anylist ANY */
-    0,  /* (374) with ::= */
+    0,  /* (29) scantok ::= */
+   -2,  /* (30) ccons ::= CONSTRAINT nm */
+   -3,  /* (31) ccons ::= DEFAULT scantok term */
+   -4,  /* (32) ccons ::= DEFAULT LP expr RP */
+   -4,  /* (33) ccons ::= DEFAULT PLUS scantok term */
+   -4,  /* (34) ccons ::= DEFAULT MINUS scantok term */
+   -3,  /* (35) ccons ::= DEFAULT scantok ID|INDEXED */
+   -3,  /* (36) ccons ::= NOT NULL onconf */
+   -5,  /* (37) ccons ::= PRIMARY KEY sortorder onconf autoinc */
+   -2,  /* (38) ccons ::= UNIQUE onconf */
+   -4,  /* (39) ccons ::= CHECK LP expr RP */
+   -4,  /* (40) ccons ::= REFERENCES nm eidlist_opt refargs */
+   -1,  /* (41) ccons ::= defer_subclause */
+   -2,  /* (42) ccons ::= COLLATE ID|STRING */
+   -3,  /* (43) generated ::= LP expr RP */
+   -4,  /* (44) generated ::= LP expr RP ID */
+    0,  /* (45) autoinc ::= */
+   -1,  /* (46) autoinc ::= AUTOINCR */
+    0,  /* (47) refargs ::= */
+   -2,  /* (48) refargs ::= refargs refarg */
+   -2,  /* (49) refarg ::= MATCH nm */
+   -3,  /* (50) refarg ::= ON INSERT refact */
+   -3,  /* (51) refarg ::= ON DELETE refact */
+   -3,  /* (52) refarg ::= ON UPDATE refact */
+   -2,  /* (53) refact ::= SET NULL */
+   -2,  /* (54) refact ::= SET DEFAULT */
+   -1,  /* (55) refact ::= CASCADE */
+   -1,  /* (56) refact ::= RESTRICT */
+   -2,  /* (57) refact ::= NO ACTION */
+   -3,  /* (58) defer_subclause ::= NOT DEFERRABLE init_deferred_pred_opt */
+   -2,  /* (59) defer_subclause ::= DEFERRABLE init_deferred_pred_opt */
+    0,  /* (60) init_deferred_pred_opt ::= */
+   -2,  /* (61) init_deferred_pred_opt ::= INITIALLY DEFERRED */
+   -2,  /* (62) init_deferred_pred_opt ::= INITIALLY IMMEDIATE */
+    0,  /* (63) conslist_opt ::= */
+   -1,  /* (64) tconscomma ::= COMMA */
+   -2,  /* (65) tcons ::= CONSTRAINT nm */
+   -7,  /* (66) tcons ::= PRIMARY KEY LP sortlist autoinc RP onconf */
+   -5,  /* (67) tcons ::= UNIQUE LP sortlist RP onconf */
+   -5,  /* (68) tcons ::= CHECK LP expr RP onconf */
+  -10,  /* (69) tcons ::= FOREIGN KEY LP eidlist RP REFERENCES nm eidlist_opt refargs defer_subclause_opt */
+    0,  /* (70) defer_subclause_opt ::= */
+    0,  /* (71) onconf ::= */
+   -3,  /* (72) onconf ::= ON CONFLICT resolvetype */
+    0,  /* (73) orconf ::= */
+   -2,  /* (74) orconf ::= OR resolvetype */
+   -1,  /* (75) resolvetype ::= IGNORE */
+   -1,  /* (76) resolvetype ::= REPLACE */
+   -4,  /* (77) cmd ::= DROP TABLE ifexists fullname */
+   -2,  /* (78) ifexists ::= IF EXISTS */
+    0,  /* (79) ifexists ::= */
+   -9,  /* (80) cmd ::= createkw temp VIEW ifnotexists nm dbnm eidlist_opt AS select */
+   -4,  /* (81) cmd ::= DROP VIEW ifexists fullname */
+   -1,  /* (82) cmd ::= select */
+   -3,  /* (83) select ::= WITH wqlist selectnowith */
+   -4,  /* (84) select ::= WITH RECURSIVE wqlist selectnowith */
+   -1,  /* (85) select ::= selectnowith */
+   -3,  /* (86) selectnowith ::= selectnowith multiselect_op oneselect */
+   -1,  /* (87) multiselect_op ::= UNION */
+   -2,  /* (88) multiselect_op ::= UNION ALL */
+   -1,  /* (89) multiselect_op ::= EXCEPT|INTERSECT */
+   -9,  /* (90) oneselect ::= SELECT distinct selcollist from where_opt groupby_opt having_opt orderby_opt limit_opt */
+  -10,  /* (91) oneselect ::= SELECT distinct selcollist from where_opt groupby_opt having_opt window_clause orderby_opt limit_opt */
+   -4,  /* (92) values ::= VALUES LP nexprlist RP */
+   -5,  /* (93) values ::= values COMMA LP nexprlist RP */
+   -1,  /* (94) distinct ::= DISTINCT */
+   -1,  /* (95) distinct ::= ALL */
+    0,  /* (96) distinct ::= */
+    0,  /* (97) sclp ::= */
+   -5,  /* (98) selcollist ::= sclp scanpt expr scanpt as */
+   -3,  /* (99) selcollist ::= sclp scanpt STAR */
+   -5,  /* (100) selcollist ::= sclp scanpt nm DOT STAR */
+   -2,  /* (101) as ::= AS nm */
+    0,  /* (102) as ::= */
+    0,  /* (103) from ::= */
+   -2,  /* (104) from ::= FROM seltablist */
+   -2,  /* (105) stl_prefix ::= seltablist joinop */
+    0,  /* (106) stl_prefix ::= */
+   -7,  /* (107) seltablist ::= stl_prefix nm dbnm as indexed_opt on_opt using_opt */
+   -9,  /* (108) seltablist ::= stl_prefix nm dbnm LP exprlist RP as on_opt using_opt */
+   -7,  /* (109) seltablist ::= stl_prefix LP select RP as on_opt using_opt */
+   -7,  /* (110) seltablist ::= stl_prefix LP seltablist RP as on_opt using_opt */
+    0,  /* (111) dbnm ::= */
+   -2,  /* (112) dbnm ::= DOT nm */
+   -1,  /* (113) fullname ::= nm */
+   -3,  /* (114) fullname ::= nm DOT nm */
+   -1,  /* (115) xfullname ::= nm */
+   -3,  /* (116) xfullname ::= nm DOT nm */
+   -5,  /* (117) xfullname ::= nm DOT nm AS nm */
+   -3,  /* (118) xfullname ::= nm AS nm */
+   -1,  /* (119) joinop ::= COMMA|JOIN */
+   -2,  /* (120) joinop ::= JOIN_KW JOIN */
+   -3,  /* (121) joinop ::= JOIN_KW nm JOIN */
+   -4,  /* (122) joinop ::= JOIN_KW nm nm JOIN */
+   -2,  /* (123) on_opt ::= ON expr */
+    0,  /* (124) on_opt ::= */
+    0,  /* (125) indexed_opt ::= */
+   -3,  /* (126) indexed_opt ::= INDEXED BY nm */
+   -2,  /* (127) indexed_opt ::= NOT INDEXED */
+   -4,  /* (128) using_opt ::= USING LP idlist RP */
+    0,  /* (129) using_opt ::= */
+    0,  /* (130) orderby_opt ::= */
+   -3,  /* (131) orderby_opt ::= ORDER BY sortlist */
+   -5,  /* (132) sortlist ::= sortlist COMMA expr sortorder nulls */
+   -3,  /* (133) sortlist ::= expr sortorder nulls */
+   -1,  /* (134) sortorder ::= ASC */
+   -1,  /* (135) sortorder ::= DESC */
+    0,  /* (136) sortorder ::= */
+   -2,  /* (137) nulls ::= NULLS FIRST */
+   -2,  /* (138) nulls ::= NULLS LAST */
+    0,  /* (139) nulls ::= */
+    0,  /* (140) groupby_opt ::= */
+   -3,  /* (141) groupby_opt ::= GROUP BY nexprlist */
+    0,  /* (142) having_opt ::= */
+   -2,  /* (143) having_opt ::= HAVING expr */
+    0,  /* (144) limit_opt ::= */
+   -2,  /* (145) limit_opt ::= LIMIT expr */
+   -4,  /* (146) limit_opt ::= LIMIT expr OFFSET expr */
+   -4,  /* (147) limit_opt ::= LIMIT expr COMMA expr */
+   -6,  /* (148) cmd ::= with DELETE FROM xfullname indexed_opt where_opt */
+    0,  /* (149) where_opt ::= */
+   -2,  /* (150) where_opt ::= WHERE expr */
+   -8,  /* (151) cmd ::= with UPDATE orconf xfullname indexed_opt SET setlist where_opt */
+   -5,  /* (152) setlist ::= setlist COMMA nm EQ expr */
+   -7,  /* (153) setlist ::= setlist COMMA LP idlist RP EQ expr */
+   -3,  /* (154) setlist ::= nm EQ expr */
+   -5,  /* (155) setlist ::= LP idlist RP EQ expr */
+   -7,  /* (156) cmd ::= with insert_cmd INTO xfullname idlist_opt select upsert */
+   -7,  /* (157) cmd ::= with insert_cmd INTO xfullname idlist_opt DEFAULT VALUES */
+    0,  /* (158) upsert ::= */
+  -11,  /* (159) upsert ::= ON CONFLICT LP sortlist RP where_opt DO UPDATE SET setlist where_opt */
+   -8,  /* (160) upsert ::= ON CONFLICT LP sortlist RP where_opt DO NOTHING */
+   -4,  /* (161) upsert ::= ON CONFLICT DO NOTHING */
+   -2,  /* (162) insert_cmd ::= INSERT orconf */
+   -1,  /* (163) insert_cmd ::= REPLACE */
+    0,  /* (164) idlist_opt ::= */
+   -3,  /* (165) idlist_opt ::= LP idlist RP */
+   -3,  /* (166) idlist ::= idlist COMMA nm */
+   -1,  /* (167) idlist ::= nm */
+   -3,  /* (168) expr ::= LP expr RP */
+   -1,  /* (169) expr ::= ID|INDEXED */
+   -1,  /* (170) expr ::= JOIN_KW */
+   -3,  /* (171) expr ::= nm DOT nm */
+   -5,  /* (172) expr ::= nm DOT nm DOT nm */
+   -1,  /* (173) term ::= NULL|FLOAT|BLOB */
+   -1,  /* (174) term ::= STRING */
+   -1,  /* (175) term ::= INTEGER */
+   -1,  /* (176) expr ::= VARIABLE */
+   -3,  /* (177) expr ::= expr COLLATE ID|STRING */
+   -6,  /* (178) expr ::= CAST LP expr AS typetoken RP */
+   -5,  /* (179) expr ::= ID|INDEXED LP distinct exprlist RP */
+   -4,  /* (180) expr ::= ID|INDEXED LP STAR RP */
+   -6,  /* (181) expr ::= ID|INDEXED LP distinct exprlist RP filter_over */
+   -5,  /* (182) expr ::= ID|INDEXED LP STAR RP filter_over */
+   -1,  /* (183) term ::= CTIME_KW */
+   -5,  /* (184) expr ::= LP nexprlist COMMA expr RP */
+   -3,  /* (185) expr ::= expr AND expr */
+   -3,  /* (186) expr ::= expr OR expr */
+   -3,  /* (187) expr ::= expr LT|GT|GE|LE expr */
+   -3,  /* (188) expr ::= expr EQ|NE expr */
+   -3,  /* (189) expr ::= expr BITAND|BITOR|LSHIFT|RSHIFT expr */
+   -3,  /* (190) expr ::= expr PLUS|MINUS expr */
+   -3,  /* (191) expr ::= expr STAR|SLASH|REM expr */
+   -3,  /* (192) expr ::= expr CONCAT expr */
+   -2,  /* (193) likeop ::= NOT LIKE_KW|MATCH */
+   -3,  /* (194) expr ::= expr likeop expr */
+   -5,  /* (195) expr ::= expr likeop expr ESCAPE expr */
+   -2,  /* (196) expr ::= expr ISNULL|NOTNULL */
+   -3,  /* (197) expr ::= expr NOT NULL */
+   -3,  /* (198) expr ::= expr IS expr */
+   -4,  /* (199) expr ::= expr IS NOT expr */
+   -2,  /* (200) expr ::= NOT expr */
+   -2,  /* (201) expr ::= BITNOT expr */
+   -2,  /* (202) expr ::= PLUS|MINUS expr */
+   -1,  /* (203) between_op ::= BETWEEN */
+   -2,  /* (204) between_op ::= NOT BETWEEN */
+   -5,  /* (205) expr ::= expr between_op expr AND expr */
+   -1,  /* (206) in_op ::= IN */
+   -2,  /* (207) in_op ::= NOT IN */
+   -5,  /* (208) expr ::= expr in_op LP exprlist RP */
+   -3,  /* (209) expr ::= LP select RP */
+   -5,  /* (210) expr ::= expr in_op LP select RP */
+   -5,  /* (211) expr ::= expr in_op nm dbnm paren_exprlist */
+   -4,  /* (212) expr ::= EXISTS LP select RP */
+   -5,  /* (213) expr ::= CASE case_operand case_exprlist case_else END */
+   -5,  /* (214) case_exprlist ::= case_exprlist WHEN expr THEN expr */
+   -4,  /* (215) case_exprlist ::= WHEN expr THEN expr */
+   -2,  /* (216) case_else ::= ELSE expr */
+    0,  /* (217) case_else ::= */
+   -1,  /* (218) case_operand ::= expr */
+    0,  /* (219) case_operand ::= */
+    0,  /* (220) exprlist ::= */
+   -3,  /* (221) nexprlist ::= nexprlist COMMA expr */
+   -1,  /* (222) nexprlist ::= expr */
+    0,  /* (223) paren_exprlist ::= */
+   -3,  /* (224) paren_exprlist ::= LP exprlist RP */
+  -12,  /* (225) cmd ::= createkw uniqueflag INDEX ifnotexists nm dbnm ON nm LP sortlist RP where_opt */
+   -1,  /* (226) uniqueflag ::= UNIQUE */
+    0,  /* (227) uniqueflag ::= */
+    0,  /* (228) eidlist_opt ::= */
+   -3,  /* (229) eidlist_opt ::= LP eidlist RP */
+   -5,  /* (230) eidlist ::= eidlist COMMA nm collate sortorder */
+   -3,  /* (231) eidlist ::= nm collate sortorder */
+    0,  /* (232) collate ::= */
+   -2,  /* (233) collate ::= COLLATE ID|STRING */
+   -4,  /* (234) cmd ::= DROP INDEX ifexists fullname */
+   -2,  /* (235) cmd ::= VACUUM vinto */
+   -3,  /* (236) cmd ::= VACUUM nm vinto */
+   -2,  /* (237) vinto ::= INTO expr */
+    0,  /* (238) vinto ::= */
+   -3,  /* (239) cmd ::= PRAGMA nm dbnm */
+   -5,  /* (240) cmd ::= PRAGMA nm dbnm EQ nmnum */
+   -6,  /* (241) cmd ::= PRAGMA nm dbnm LP nmnum RP */
+   -5,  /* (242) cmd ::= PRAGMA nm dbnm EQ minus_num */
+   -6,  /* (243) cmd ::= PRAGMA nm dbnm LP minus_num RP */
+   -2,  /* (244) plus_num ::= PLUS INTEGER|FLOAT */
+   -2,  /* (245) minus_num ::= MINUS INTEGER|FLOAT */
+   -5,  /* (246) cmd ::= createkw trigger_decl BEGIN trigger_cmd_list END */
+  -11,  /* (247) trigger_decl ::= temp TRIGGER ifnotexists nm dbnm trigger_time trigger_event ON fullname foreach_clause when_clause */
+   -1,  /* (248) trigger_time ::= BEFORE|AFTER */
+   -2,  /* (249) trigger_time ::= INSTEAD OF */
+    0,  /* (250) trigger_time ::= */
+   -1,  /* (251) trigger_event ::= DELETE|INSERT */
+   -1,  /* (252) trigger_event ::= UPDATE */
+   -3,  /* (253) trigger_event ::= UPDATE OF idlist */
+    0,  /* (254) when_clause ::= */
+   -2,  /* (255) when_clause ::= WHEN expr */
+   -3,  /* (256) trigger_cmd_list ::= trigger_cmd_list trigger_cmd SEMI */
+   -2,  /* (257) trigger_cmd_list ::= trigger_cmd SEMI */
+   -3,  /* (258) trnm ::= nm DOT nm */
+   -3,  /* (259) tridxby ::= INDEXED BY nm */
+   -2,  /* (260) tridxby ::= NOT INDEXED */
+   -8,  /* (261) trigger_cmd ::= UPDATE orconf trnm tridxby SET setlist where_opt scanpt */
+   -8,  /* (262) trigger_cmd ::= scanpt insert_cmd INTO trnm idlist_opt select upsert scanpt */
+   -6,  /* (263) trigger_cmd ::= DELETE FROM trnm tridxby where_opt scanpt */
+   -3,  /* (264) trigger_cmd ::= scanpt select scanpt */
+   -4,  /* (265) expr ::= RAISE LP IGNORE RP */
+   -6,  /* (266) expr ::= RAISE LP raisetype COMMA nm RP */
+   -1,  /* (267) raisetype ::= ROLLBACK */
+   -1,  /* (268) raisetype ::= ABORT */
+   -1,  /* (269) raisetype ::= FAIL */
+   -4,  /* (270) cmd ::= DROP TRIGGER ifexists fullname */
+   -6,  /* (271) cmd ::= ATTACH database_kw_opt expr AS expr key_opt */
+   -3,  /* (272) cmd ::= DETACH database_kw_opt expr */
+    0,  /* (273) key_opt ::= */
+   -2,  /* (274) key_opt ::= KEY expr */
+   -1,  /* (275) cmd ::= REINDEX */
+   -3,  /* (276) cmd ::= REINDEX nm dbnm */
+   -1,  /* (277) cmd ::= ANALYZE */
+   -3,  /* (278) cmd ::= ANALYZE nm dbnm */
+   -6,  /* (279) cmd ::= ALTER TABLE fullname RENAME TO nm */
+   -7,  /* (280) cmd ::= ALTER TABLE add_column_fullname ADD kwcolumn_opt columnname carglist */
+   -1,  /* (281) add_column_fullname ::= fullname */
+   -8,  /* (282) cmd ::= ALTER TABLE fullname RENAME kwcolumn_opt nm TO nm */
+   -1,  /* (283) cmd ::= create_vtab */
+   -4,  /* (284) cmd ::= create_vtab LP vtabarglist RP */
+   -8,  /* (285) create_vtab ::= createkw VIRTUAL TABLE ifnotexists nm dbnm USING nm */
+    0,  /* (286) vtabarg ::= */
+   -1,  /* (287) vtabargtoken ::= ANY */
+   -3,  /* (288) vtabargtoken ::= lp anylist RP */
+   -1,  /* (289) lp ::= LP */
+   -2,  /* (290) with ::= WITH wqlist */
+   -3,  /* (291) with ::= WITH RECURSIVE wqlist */
+   -6,  /* (292) wqlist ::= nm eidlist_opt AS LP select RP */
+   -8,  /* (293) wqlist ::= wqlist COMMA nm eidlist_opt AS LP select RP */
+   -1,  /* (294) windowdefn_list ::= windowdefn */
+   -3,  /* (295) windowdefn_list ::= windowdefn_list COMMA windowdefn */
+   -5,  /* (296) windowdefn ::= nm AS LP window RP */
+   -5,  /* (297) window ::= PARTITION BY nexprlist orderby_opt frame_opt */
+   -6,  /* (298) window ::= nm PARTITION BY nexprlist orderby_opt frame_opt */
+   -4,  /* (299) window ::= ORDER BY sortlist frame_opt */
+   -5,  /* (300) window ::= nm ORDER BY sortlist frame_opt */
+   -1,  /* (301) window ::= frame_opt */
+   -2,  /* (302) window ::= nm frame_opt */
+    0,  /* (303) frame_opt ::= */
+   -3,  /* (304) frame_opt ::= range_or_rows frame_bound_s frame_exclude_opt */
+   -6,  /* (305) frame_opt ::= range_or_rows BETWEEN frame_bound_s AND frame_bound_e frame_exclude_opt */
+   -1,  /* (306) range_or_rows ::= RANGE|ROWS|GROUPS */
+   -1,  /* (307) frame_bound_s ::= frame_bound */
+   -2,  /* (308) frame_bound_s ::= UNBOUNDED PRECEDING */
+   -1,  /* (309) frame_bound_e ::= frame_bound */
+   -2,  /* (310) frame_bound_e ::= UNBOUNDED FOLLOWING */
+   -2,  /* (311) frame_bound ::= expr PRECEDING|FOLLOWING */
+   -2,  /* (312) frame_bound ::= CURRENT ROW */
+    0,  /* (313) frame_exclude_opt ::= */
+   -2,  /* (314) frame_exclude_opt ::= EXCLUDE frame_exclude */
+   -2,  /* (315) frame_exclude ::= NO OTHERS */
+   -2,  /* (316) frame_exclude ::= CURRENT ROW */
+   -1,  /* (317) frame_exclude ::= GROUP|TIES */
+   -2,  /* (318) window_clause ::= WINDOW windowdefn_list */
+   -2,  /* (319) filter_over ::= filter_clause over_clause */
+   -1,  /* (320) filter_over ::= over_clause */
+   -1,  /* (321) filter_over ::= filter_clause */
+   -4,  /* (322) over_clause ::= OVER LP window RP */
+   -2,  /* (323) over_clause ::= OVER nm */
+   -5,  /* (324) filter_clause ::= FILTER LP WHERE expr RP */
+   -1,  /* (325) input ::= cmdlist */
+   -2,  /* (326) cmdlist ::= cmdlist ecmd */
+   -1,  /* (327) cmdlist ::= ecmd */
+   -1,  /* (328) ecmd ::= SEMI */
+   -2,  /* (329) ecmd ::= cmdx SEMI */
+   -3,  /* (330) ecmd ::= explain cmdx SEMI */
+    0,  /* (331) trans_opt ::= */
+   -1,  /* (332) trans_opt ::= TRANSACTION */
+   -2,  /* (333) trans_opt ::= TRANSACTION nm */
+   -1,  /* (334) savepoint_opt ::= SAVEPOINT */
+    0,  /* (335) savepoint_opt ::= */
+   -2,  /* (336) cmd ::= create_table create_table_args */
+   -4,  /* (337) columnlist ::= columnlist COMMA columnname carglist */
+   -2,  /* (338) columnlist ::= columnname carglist */
+   -1,  /* (339) nm ::= ID|INDEXED */
+   -1,  /* (340) nm ::= STRING */
+   -1,  /* (341) nm ::= JOIN_KW */
+   -1,  /* (342) typetoken ::= typename */
+   -1,  /* (343) typename ::= ID|STRING */
+   -1,  /* (344) signed ::= plus_num */
+   -1,  /* (345) signed ::= minus_num */
+   -2,  /* (346) carglist ::= carglist ccons */
+    0,  /* (347) carglist ::= */
+   -2,  /* (348) ccons ::= NULL onconf */
+   -4,  /* (349) ccons ::= GENERATED ALWAYS AS generated */
+   -2,  /* (350) ccons ::= AS generated */
+   -2,  /* (351) conslist_opt ::= COMMA conslist */
+   -3,  /* (352) conslist ::= conslist tconscomma tcons */
+   -1,  /* (353) conslist ::= tcons */
+    0,  /* (354) tconscomma ::= */
+   -1,  /* (355) defer_subclause_opt ::= defer_subclause */
+   -1,  /* (356) resolvetype ::= raisetype */
+   -1,  /* (357) selectnowith ::= oneselect */
+   -1,  /* (358) oneselect ::= values */
+   -2,  /* (359) sclp ::= selcollist COMMA */
+   -1,  /* (360) as ::= ID|STRING */
+   -1,  /* (361) expr ::= term */
+   -1,  /* (362) likeop ::= LIKE_KW|MATCH */
+   -1,  /* (363) exprlist ::= nexprlist */
+   -1,  /* (364) nmnum ::= plus_num */
+   -1,  /* (365) nmnum ::= nm */
+   -1,  /* (366) nmnum ::= ON */
+   -1,  /* (367) nmnum ::= DELETE */
+   -1,  /* (368) nmnum ::= DEFAULT */
+   -1,  /* (369) plus_num ::= INTEGER|FLOAT */
+    0,  /* (370) foreach_clause ::= */
+   -3,  /* (371) foreach_clause ::= FOR EACH ROW */
+   -1,  /* (372) trnm ::= nm */
+    0,  /* (373) tridxby ::= */
+   -1,  /* (374) database_kw_opt ::= DATABASE */
+    0,  /* (375) database_kw_opt ::= */
+    0,  /* (376) kwcolumn_opt ::= */
+   -1,  /* (377) kwcolumn_opt ::= COLUMNKW */
+   -1,  /* (378) vtabarglist ::= vtabarg */
+   -3,  /* (379) vtabarglist ::= vtabarglist COMMA vtabarg */
+   -2,  /* (380) vtabarg ::= vtabarg vtabargtoken */
+    0,  /* (381) anylist ::= */
+   -4,  /* (382) anylist ::= anylist LP anylist RP */
+   -2,  /* (383) anylist ::= anylist ANY */
+    0,  /* (384) with ::= */
 };
 
 static void yy_accept(yyParser*);  /* Forward Declaration */
@@ -151258,12 +155549,15 @@ static YYACTIONTYPE yy_reduce(
   if( yyTraceFILE && yyruleno<(int)(sizeof(yyRuleName)/sizeof(yyRuleName[0])) ){
     yysize = yyRuleInfoNRhs[yyruleno];
     if( yysize ){
-      fprintf(yyTraceFILE, "%sReduce %d [%s], go to state %d.\n",
+      fprintf(yyTraceFILE, "%sReduce %d [%s]%s, pop back to state %d.\n",
         yyTracePrompt,
-        yyruleno, yyRuleName[yyruleno], yymsp[yysize].stateno);
+        yyruleno, yyRuleName[yyruleno],
+        yyruleno<YYNRULE_WITH_ACTION ? "" : " without external action",
+        yymsp[yysize].stateno);
     }else{
-      fprintf(yyTraceFILE, "%sReduce %d [%s].\n",
-        yyTracePrompt, yyruleno, yyRuleName[yyruleno]);
+      fprintf(yyTraceFILE, "%sReduce %d [%s]%s.\n",
+        yyTracePrompt, yyruleno, yyRuleName[yyruleno],
+        yyruleno<YYNRULE_WITH_ACTION ? "" : " without external action");
     }
   }
 #endif /* NDEBUG */
@@ -151321,16 +155615,16 @@ static YYACTIONTYPE yy_reduce(
 { sqlite3FinishCoding(pParse); }
         break;
       case 3: /* cmd ::= BEGIN transtype trans_opt */
-{sqlite3BeginTransaction(pParse, yymsp[-1].minor.yy494);}
+{sqlite3BeginTransaction(pParse, yymsp[-1].minor.yy192);}
         break;
       case 4: /* transtype ::= */
-{yymsp[1].minor.yy494 = TK_DEFERRED;}
+{yymsp[1].minor.yy192 = TK_DEFERRED;}
         break;
       case 5: /* transtype ::= DEFERRED */
       case 6: /* transtype ::= IMMEDIATE */ yytestcase(yyruleno==6);
       case 7: /* transtype ::= EXCLUSIVE */ yytestcase(yyruleno==7);
-      case 300: /* range_or_rows ::= RANGE|ROWS|GROUPS */ yytestcase(yyruleno==300);
-{yymsp[0].minor.yy494 = yymsp[0].major; /*A-overwrites-X*/}
+      case 306: /* range_or_rows ::= RANGE|ROWS|GROUPS */ yytestcase(yyruleno==306);
+{yymsp[0].minor.yy192 = yymsp[0].major; /*A-overwrites-X*/}
         break;
       case 8: /* cmd ::= COMMIT|END trans_opt */
       case 9: /* cmd ::= ROLLBACK trans_opt */ yytestcase(yyruleno==9);
@@ -151353,7 +155647,7 @@ static YYACTIONTYPE yy_reduce(
         break;
       case 13: /* create_table ::= createkw temp TABLE ifnotexists nm dbnm */
 {
-   sqlite3StartTable(pParse,&yymsp[-1].minor.yy0,&yymsp[0].minor.yy0,yymsp[-4].minor.yy494,0,0,yymsp[-2].minor.yy494);
+   sqlite3StartTable(pParse,&yymsp[-1].minor.yy0,&yymsp[0].minor.yy0,yymsp[-4].minor.yy192,0,0,yymsp[-2].minor.yy192);
 }
         break;
       case 14: /* createkw ::= CREATE */
@@ -151362,38 +155656,38 @@ static YYACTIONTYPE yy_reduce(
       case 15: /* ifnotexists ::= */
       case 18: /* temp ::= */ yytestcase(yyruleno==18);
       case 21: /* table_options ::= */ yytestcase(yyruleno==21);
-      case 42: /* autoinc ::= */ yytestcase(yyruleno==42);
-      case 57: /* init_deferred_pred_opt ::= */ yytestcase(yyruleno==57);
-      case 67: /* defer_subclause_opt ::= */ yytestcase(yyruleno==67);
-      case 76: /* ifexists ::= */ yytestcase(yyruleno==76);
-      case 93: /* distinct ::= */ yytestcase(yyruleno==93);
-      case 226: /* collate ::= */ yytestcase(yyruleno==226);
-{yymsp[1].minor.yy494 = 0;}
+      case 45: /* autoinc ::= */ yytestcase(yyruleno==45);
+      case 60: /* init_deferred_pred_opt ::= */ yytestcase(yyruleno==60);
+      case 70: /* defer_subclause_opt ::= */ yytestcase(yyruleno==70);
+      case 79: /* ifexists ::= */ yytestcase(yyruleno==79);
+      case 96: /* distinct ::= */ yytestcase(yyruleno==96);
+      case 232: /* collate ::= */ yytestcase(yyruleno==232);
+{yymsp[1].minor.yy192 = 0;}
         break;
       case 16: /* ifnotexists ::= IF NOT EXISTS */
-{yymsp[-2].minor.yy494 = 1;}
+{yymsp[-2].minor.yy192 = 1;}
         break;
       case 17: /* temp ::= TEMP */
-      case 43: /* autoinc ::= AUTOINCR */ yytestcase(yyruleno==43);
-{yymsp[0].minor.yy494 = 1;}
+      case 46: /* autoinc ::= AUTOINCR */ yytestcase(yyruleno==46);
+{yymsp[0].minor.yy192 = 1;}
         break;
       case 19: /* create_table_args ::= LP columnlist conslist_opt RP table_options */
 {
-  sqlite3EndTable(pParse,&yymsp[-2].minor.yy0,&yymsp[-1].minor.yy0,yymsp[0].minor.yy494,0);
+  sqlite3EndTable(pParse,&yymsp[-2].minor.yy0,&yymsp[-1].minor.yy0,yymsp[0].minor.yy192,0);
 }
         break;
       case 20: /* create_table_args ::= AS select */
 {
-  sqlite3EndTable(pParse,0,0,0,yymsp[0].minor.yy457);
-  sqlite3SelectDelete(pParse->db, yymsp[0].minor.yy457);
+  sqlite3EndTable(pParse,0,0,0,yymsp[0].minor.yy539);
+  sqlite3SelectDelete(pParse->db, yymsp[0].minor.yy539);
 }
         break;
       case 22: /* table_options ::= WITHOUT nm */
 {
   if( yymsp[0].minor.yy0.n==5 && sqlite3_strnicmp(yymsp[0].minor.yy0.z,"rowid",5)==0 ){
-    yymsp[-1].minor.yy494 = TF_WithoutRowid | TF_NoVisibleRowid;
+    yymsp[-1].minor.yy192 = TF_WithoutRowid | TF_NoVisibleRowid;
   }else{
-    yymsp[-1].minor.yy494 = 0;
+    yymsp[-1].minor.yy192 = 0;
     sqlite3ErrorMsg(pParse, "unknown table option: %.*s", yymsp[0].minor.yy0.n, yymsp[0].minor.yy0.z);
   }
 }
@@ -151402,8 +155696,8 @@ static YYACTIONTYPE yy_reduce(
 {sqlite3AddColumn(pParse,&yymsp[-1].minor.yy0,&yymsp[0].minor.yy0);}
         break;
       case 24: /* typetoken ::= */
-      case 60: /* conslist_opt ::= */ yytestcase(yyruleno==60);
-      case 99: /* as ::= */ yytestcase(yyruleno==99);
+      case 63: /* conslist_opt ::= */ yytestcase(yyruleno==63);
+      case 102: /* as ::= */ yytestcase(yyruleno==102);
 {yymsp[1].minor.yy0.n = 0; yymsp[1].minor.yy0.z = 0;}
         break;
       case 25: /* typetoken ::= typename LP signed RP */
@@ -151422,29 +155716,35 @@ static YYACTIONTYPE yy_reduce(
       case 28: /* scanpt ::= */
 {
   assert( yyLookahead!=YYNOCODE );
-  yymsp[1].minor.yy294 = yyLookaheadToken.z;
+  yymsp[1].minor.yy436 = yyLookaheadToken.z;
 }
         break;
-      case 29: /* ccons ::= CONSTRAINT nm */
-      case 62: /* tcons ::= CONSTRAINT nm */ yytestcase(yyruleno==62);
+      case 29: /* scantok ::= */
+{
+  assert( yyLookahead!=YYNOCODE );
+  yymsp[1].minor.yy0 = yyLookaheadToken;
+}
+        break;
+      case 30: /* ccons ::= CONSTRAINT nm */
+      case 65: /* tcons ::= CONSTRAINT nm */ yytestcase(yyruleno==65);
 {pParse->constraintName = yymsp[0].minor.yy0;}
         break;
-      case 30: /* ccons ::= DEFAULT scanpt term scanpt */
-{sqlite3AddDefaultValue(pParse,yymsp[-1].minor.yy524,yymsp[-2].minor.yy294,yymsp[0].minor.yy294);}
+      case 31: /* ccons ::= DEFAULT scantok term */
+{sqlite3AddDefaultValue(pParse,yymsp[0].minor.yy202,yymsp[-1].minor.yy0.z,&yymsp[-1].minor.yy0.z[yymsp[-1].minor.yy0.n]);}
         break;
-      case 31: /* ccons ::= DEFAULT LP expr RP */
-{sqlite3AddDefaultValue(pParse,yymsp[-1].minor.yy524,yymsp[-2].minor.yy0.z+1,yymsp[0].minor.yy0.z);}
+      case 32: /* ccons ::= DEFAULT LP expr RP */
+{sqlite3AddDefaultValue(pParse,yymsp[-1].minor.yy202,yymsp[-2].minor.yy0.z+1,yymsp[0].minor.yy0.z);}
         break;
-      case 32: /* ccons ::= DEFAULT PLUS term scanpt */
-{sqlite3AddDefaultValue(pParse,yymsp[-1].minor.yy524,yymsp[-2].minor.yy0.z,yymsp[0].minor.yy294);}
+      case 33: /* ccons ::= DEFAULT PLUS scantok term */
+{sqlite3AddDefaultValue(pParse,yymsp[0].minor.yy202,yymsp[-2].minor.yy0.z,&yymsp[-1].minor.yy0.z[yymsp[-1].minor.yy0.n]);}
         break;
-      case 33: /* ccons ::= DEFAULT MINUS term scanpt */
+      case 34: /* ccons ::= DEFAULT MINUS scantok term */
 {
-  Expr *p = sqlite3PExpr(pParse, TK_UMINUS, yymsp[-1].minor.yy524, 0);
-  sqlite3AddDefaultValue(pParse,p,yymsp[-2].minor.yy0.z,yymsp[0].minor.yy294);
+  Expr *p = sqlite3PExpr(pParse, TK_UMINUS, yymsp[0].minor.yy202, 0);
+  sqlite3AddDefaultValue(pParse,p,yymsp[-2].minor.yy0.z,&yymsp[-1].minor.yy0.z[yymsp[-1].minor.yy0.n]);
 }
         break;
-      case 34: /* ccons ::= DEFAULT scanpt ID|INDEXED */
+      case 35: /* ccons ::= DEFAULT scantok ID|INDEXED */
 {
   Expr *p = tokenExpr(pParse, TK_STRING, yymsp[0].minor.yy0);
   if( p ){
@@ -151454,171 +155754,177 @@ static YYACTIONTYPE yy_reduce(
     sqlite3AddDefaultValue(pParse,p,yymsp[0].minor.yy0.z,yymsp[0].minor.yy0.z+yymsp[0].minor.yy0.n);
 }
         break;
-      case 35: /* ccons ::= NOT NULL onconf */
-{sqlite3AddNotNull(pParse, yymsp[0].minor.yy494);}
+      case 36: /* ccons ::= NOT NULL onconf */
+{sqlite3AddNotNull(pParse, yymsp[0].minor.yy192);}
         break;
-      case 36: /* ccons ::= PRIMARY KEY sortorder onconf autoinc */
-{sqlite3AddPrimaryKey(pParse,0,yymsp[-1].minor.yy494,yymsp[0].minor.yy494,yymsp[-2].minor.yy494);}
+      case 37: /* ccons ::= PRIMARY KEY sortorder onconf autoinc */
+{sqlite3AddPrimaryKey(pParse,0,yymsp[-1].minor.yy192,yymsp[0].minor.yy192,yymsp[-2].minor.yy192);}
         break;
-      case 37: /* ccons ::= UNIQUE onconf */
-{sqlite3CreateIndex(pParse,0,0,0,0,yymsp[0].minor.yy494,0,0,0,0,
+      case 38: /* ccons ::= UNIQUE onconf */
+{sqlite3CreateIndex(pParse,0,0,0,0,yymsp[0].minor.yy192,0,0,0,0,
                                    SQLITE_IDXTYPE_UNIQUE);}
         break;
-      case 38: /* ccons ::= CHECK LP expr RP */
-{sqlite3AddCheckConstraint(pParse,yymsp[-1].minor.yy524);}
+      case 39: /* ccons ::= CHECK LP expr RP */
+{sqlite3AddCheckConstraint(pParse,yymsp[-1].minor.yy202);}
         break;
-      case 39: /* ccons ::= REFERENCES nm eidlist_opt refargs */
-{sqlite3CreateForeignKey(pParse,0,&yymsp[-2].minor.yy0,yymsp[-1].minor.yy434,yymsp[0].minor.yy494);}
+      case 40: /* ccons ::= REFERENCES nm eidlist_opt refargs */
+{sqlite3CreateForeignKey(pParse,0,&yymsp[-2].minor.yy0,yymsp[-1].minor.yy242,yymsp[0].minor.yy192);}
         break;
-      case 40: /* ccons ::= defer_subclause */
-{sqlite3DeferForeignKey(pParse,yymsp[0].minor.yy494);}
+      case 41: /* ccons ::= defer_subclause */
+{sqlite3DeferForeignKey(pParse,yymsp[0].minor.yy192);}
         break;
-      case 41: /* ccons ::= COLLATE ID|STRING */
+      case 42: /* ccons ::= COLLATE ID|STRING */
 {sqlite3AddCollateType(pParse, &yymsp[0].minor.yy0);}
         break;
-      case 44: /* refargs ::= */
-{ yymsp[1].minor.yy494 = OE_None*0x0101; /* EV: R-19803-45884 */}
+      case 43: /* generated ::= LP expr RP */
+{sqlite3AddGenerated(pParse,yymsp[-1].minor.yy202,0);}
+        break;
+      case 44: /* generated ::= LP expr RP ID */
+{sqlite3AddGenerated(pParse,yymsp[-2].minor.yy202,&yymsp[0].minor.yy0);}
+        break;
+      case 47: /* refargs ::= */
+{ yymsp[1].minor.yy192 = OE_None*0x0101; /* EV: R-19803-45884 */}
         break;
-      case 45: /* refargs ::= refargs refarg */
-{ yymsp[-1].minor.yy494 = (yymsp[-1].minor.yy494 & ~yymsp[0].minor.yy355.mask) | yymsp[0].minor.yy355.value; }
+      case 48: /* refargs ::= refargs refarg */
+{ yymsp[-1].minor.yy192 = (yymsp[-1].minor.yy192 & ~yymsp[0].minor.yy207.mask) | yymsp[0].minor.yy207.value; }
         break;
-      case 46: /* refarg ::= MATCH nm */
-{ yymsp[-1].minor.yy355.value = 0;     yymsp[-1].minor.yy355.mask = 0x000000; }
+      case 49: /* refarg ::= MATCH nm */
+{ yymsp[-1].minor.yy207.value = 0;     yymsp[-1].minor.yy207.mask = 0x000000; }
         break;
-      case 47: /* refarg ::= ON INSERT refact */
-{ yymsp[-2].minor.yy355.value = 0;     yymsp[-2].minor.yy355.mask = 0x000000; }
+      case 50: /* refarg ::= ON INSERT refact */
+{ yymsp[-2].minor.yy207.value = 0;     yymsp[-2].minor.yy207.mask = 0x000000; }
         break;
-      case 48: /* refarg ::= ON DELETE refact */
-{ yymsp[-2].minor.yy355.value = yymsp[0].minor.yy494;     yymsp[-2].minor.yy355.mask = 0x0000ff; }
+      case 51: /* refarg ::= ON DELETE refact */
+{ yymsp[-2].minor.yy207.value = yymsp[0].minor.yy192;     yymsp[-2].minor.yy207.mask = 0x0000ff; }
         break;
-      case 49: /* refarg ::= ON UPDATE refact */
-{ yymsp[-2].minor.yy355.value = yymsp[0].minor.yy494<<8;  yymsp[-2].minor.yy355.mask = 0x00ff00; }
+      case 52: /* refarg ::= ON UPDATE refact */
+{ yymsp[-2].minor.yy207.value = yymsp[0].minor.yy192<<8;  yymsp[-2].minor.yy207.mask = 0x00ff00; }
         break;
-      case 50: /* refact ::= SET NULL */
-{ yymsp[-1].minor.yy494 = OE_SetNull;  /* EV: R-33326-45252 */}
+      case 53: /* refact ::= SET NULL */
+{ yymsp[-1].minor.yy192 = OE_SetNull;  /* EV: R-33326-45252 */}
         break;
-      case 51: /* refact ::= SET DEFAULT */
-{ yymsp[-1].minor.yy494 = OE_SetDflt;  /* EV: R-33326-45252 */}
+      case 54: /* refact ::= SET DEFAULT */
+{ yymsp[-1].minor.yy192 = OE_SetDflt;  /* EV: R-33326-45252 */}
         break;
-      case 52: /* refact ::= CASCADE */
-{ yymsp[0].minor.yy494 = OE_Cascade;  /* EV: R-33326-45252 */}
+      case 55: /* refact ::= CASCADE */
+{ yymsp[0].minor.yy192 = OE_Cascade;  /* EV: R-33326-45252 */}
         break;
-      case 53: /* refact ::= RESTRICT */
-{ yymsp[0].minor.yy494 = OE_Restrict; /* EV: R-33326-45252 */}
+      case 56: /* refact ::= RESTRICT */
+{ yymsp[0].minor.yy192 = OE_Restrict; /* EV: R-33326-45252 */}
         break;
-      case 54: /* refact ::= NO ACTION */
-{ yymsp[-1].minor.yy494 = OE_None;     /* EV: R-33326-45252 */}
+      case 57: /* refact ::= NO ACTION */
+{ yymsp[-1].minor.yy192 = OE_None;     /* EV: R-33326-45252 */}
         break;
-      case 55: /* defer_subclause ::= NOT DEFERRABLE init_deferred_pred_opt */
-{yymsp[-2].minor.yy494 = 0;}
+      case 58: /* defer_subclause ::= NOT DEFERRABLE init_deferred_pred_opt */
+{yymsp[-2].minor.yy192 = 0;}
         break;
-      case 56: /* defer_subclause ::= DEFERRABLE init_deferred_pred_opt */
-      case 71: /* orconf ::= OR resolvetype */ yytestcase(yyruleno==71);
-      case 156: /* insert_cmd ::= INSERT orconf */ yytestcase(yyruleno==156);
-{yymsp[-1].minor.yy494 = yymsp[0].minor.yy494;}
+      case 59: /* defer_subclause ::= DEFERRABLE init_deferred_pred_opt */
+      case 74: /* orconf ::= OR resolvetype */ yytestcase(yyruleno==74);
+      case 162: /* insert_cmd ::= INSERT orconf */ yytestcase(yyruleno==162);
+{yymsp[-1].minor.yy192 = yymsp[0].minor.yy192;}
         break;
-      case 58: /* init_deferred_pred_opt ::= INITIALLY DEFERRED */
-      case 75: /* ifexists ::= IF EXISTS */ yytestcase(yyruleno==75);
-      case 198: /* between_op ::= NOT BETWEEN */ yytestcase(yyruleno==198);
-      case 201: /* in_op ::= NOT IN */ yytestcase(yyruleno==201);
-      case 227: /* collate ::= COLLATE ID|STRING */ yytestcase(yyruleno==227);
-{yymsp[-1].minor.yy494 = 1;}
+      case 61: /* init_deferred_pred_opt ::= INITIALLY DEFERRED */
+      case 78: /* ifexists ::= IF EXISTS */ yytestcase(yyruleno==78);
+      case 204: /* between_op ::= NOT BETWEEN */ yytestcase(yyruleno==204);
+      case 207: /* in_op ::= NOT IN */ yytestcase(yyruleno==207);
+      case 233: /* collate ::= COLLATE ID|STRING */ yytestcase(yyruleno==233);
+{yymsp[-1].minor.yy192 = 1;}
         break;
-      case 59: /* init_deferred_pred_opt ::= INITIALLY IMMEDIATE */
-{yymsp[-1].minor.yy494 = 0;}
+      case 62: /* init_deferred_pred_opt ::= INITIALLY IMMEDIATE */
+{yymsp[-1].minor.yy192 = 0;}
         break;
-      case 61: /* tconscomma ::= COMMA */
+      case 64: /* tconscomma ::= COMMA */
 {pParse->constraintName.n = 0;}
         break;
-      case 63: /* tcons ::= PRIMARY KEY LP sortlist autoinc RP onconf */
-{sqlite3AddPrimaryKey(pParse,yymsp[-3].minor.yy434,yymsp[0].minor.yy494,yymsp[-2].minor.yy494,0);}
+      case 66: /* tcons ::= PRIMARY KEY LP sortlist autoinc RP onconf */
+{sqlite3AddPrimaryKey(pParse,yymsp[-3].minor.yy242,yymsp[0].minor.yy192,yymsp[-2].minor.yy192,0);}
         break;
-      case 64: /* tcons ::= UNIQUE LP sortlist RP onconf */
-{sqlite3CreateIndex(pParse,0,0,0,yymsp[-2].minor.yy434,yymsp[0].minor.yy494,0,0,0,0,
+      case 67: /* tcons ::= UNIQUE LP sortlist RP onconf */
+{sqlite3CreateIndex(pParse,0,0,0,yymsp[-2].minor.yy242,yymsp[0].minor.yy192,0,0,0,0,
                                        SQLITE_IDXTYPE_UNIQUE);}
         break;
-      case 65: /* tcons ::= CHECK LP expr RP onconf */
-{sqlite3AddCheckConstraint(pParse,yymsp[-2].minor.yy524);}
+      case 68: /* tcons ::= CHECK LP expr RP onconf */
+{sqlite3AddCheckConstraint(pParse,yymsp[-2].minor.yy202);}
         break;
-      case 66: /* tcons ::= FOREIGN KEY LP eidlist RP REFERENCES nm eidlist_opt refargs defer_subclause_opt */
+      case 69: /* tcons ::= FOREIGN KEY LP eidlist RP REFERENCES nm eidlist_opt refargs defer_subclause_opt */
 {
-    sqlite3CreateForeignKey(pParse, yymsp[-6].minor.yy434, &yymsp[-3].minor.yy0, yymsp[-2].minor.yy434, yymsp[-1].minor.yy494);
-    sqlite3DeferForeignKey(pParse, yymsp[0].minor.yy494);
+    sqlite3CreateForeignKey(pParse, yymsp[-6].minor.yy242, &yymsp[-3].minor.yy0, yymsp[-2].minor.yy242, yymsp[-1].minor.yy192);
+    sqlite3DeferForeignKey(pParse, yymsp[0].minor.yy192);
 }
         break;
-      case 68: /* onconf ::= */
-      case 70: /* orconf ::= */ yytestcase(yyruleno==70);
-{yymsp[1].minor.yy494 = OE_Default;}
+      case 71: /* onconf ::= */
+      case 73: /* orconf ::= */ yytestcase(yyruleno==73);
+{yymsp[1].minor.yy192 = OE_Default;}
         break;
-      case 69: /* onconf ::= ON CONFLICT resolvetype */
-{yymsp[-2].minor.yy494 = yymsp[0].minor.yy494;}
+      case 72: /* onconf ::= ON CONFLICT resolvetype */
+{yymsp[-2].minor.yy192 = yymsp[0].minor.yy192;}
         break;
-      case 72: /* resolvetype ::= IGNORE */
-{yymsp[0].minor.yy494 = OE_Ignore;}
+      case 75: /* resolvetype ::= IGNORE */
+{yymsp[0].minor.yy192 = OE_Ignore;}
         break;
-      case 73: /* resolvetype ::= REPLACE */
-      case 157: /* insert_cmd ::= REPLACE */ yytestcase(yyruleno==157);
-{yymsp[0].minor.yy494 = OE_Replace;}
+      case 76: /* resolvetype ::= REPLACE */
+      case 163: /* insert_cmd ::= REPLACE */ yytestcase(yyruleno==163);
+{yymsp[0].minor.yy192 = OE_Replace;}
         break;
-      case 74: /* cmd ::= DROP TABLE ifexists fullname */
+      case 77: /* cmd ::= DROP TABLE ifexists fullname */
 {
-  sqlite3DropTable(pParse, yymsp[0].minor.yy483, 0, yymsp[-1].minor.yy494);
+  sqlite3DropTable(pParse, yymsp[0].minor.yy47, 0, yymsp[-1].minor.yy192);
 }
         break;
-      case 77: /* cmd ::= createkw temp VIEW ifnotexists nm dbnm eidlist_opt AS select */
+      case 80: /* cmd ::= createkw temp VIEW ifnotexists nm dbnm eidlist_opt AS select */
 {
-  sqlite3CreateView(pParse, &yymsp[-8].minor.yy0, &yymsp[-4].minor.yy0, &yymsp[-3].minor.yy0, yymsp[-2].minor.yy434, yymsp[0].minor.yy457, yymsp[-7].minor.yy494, yymsp[-5].minor.yy494);
+  sqlite3CreateView(pParse, &yymsp[-8].minor.yy0, &yymsp[-4].minor.yy0, &yymsp[-3].minor.yy0, yymsp[-2].minor.yy242, yymsp[0].minor.yy539, yymsp[-7].minor.yy192, yymsp[-5].minor.yy192);
 }
         break;
-      case 78: /* cmd ::= DROP VIEW ifexists fullname */
+      case 81: /* cmd ::= DROP VIEW ifexists fullname */
 {
-  sqlite3DropTable(pParse, yymsp[0].minor.yy483, 1, yymsp[-1].minor.yy494);
+  sqlite3DropTable(pParse, yymsp[0].minor.yy47, 1, yymsp[-1].minor.yy192);
 }
         break;
-      case 79: /* cmd ::= select */
+      case 82: /* cmd ::= select */
 {
   SelectDest dest = {SRT_Output, 0, 0, 0, 0, 0};
-  sqlite3Select(pParse, yymsp[0].minor.yy457, &dest);
-  sqlite3SelectDelete(pParse->db, yymsp[0].minor.yy457);
+  sqlite3Select(pParse, yymsp[0].minor.yy539, &dest);
+  sqlite3SelectDelete(pParse->db, yymsp[0].minor.yy539);
 }
         break;
-      case 80: /* select ::= WITH wqlist selectnowith */
+      case 83: /* select ::= WITH wqlist selectnowith */
 {
-  Select *p = yymsp[0].minor.yy457;
+  Select *p = yymsp[0].minor.yy539;
   if( p ){
-    p->pWith = yymsp[-1].minor.yy59;
+    p->pWith = yymsp[-1].minor.yy131;
     parserDoubleLinkSelect(pParse, p);
   }else{
-    sqlite3WithDelete(pParse->db, yymsp[-1].minor.yy59);
+    sqlite3WithDelete(pParse->db, yymsp[-1].minor.yy131);
   }
-  yymsp[-2].minor.yy457 = p;
+  yymsp[-2].minor.yy539 = p;
 }
         break;
-      case 81: /* select ::= WITH RECURSIVE wqlist selectnowith */
+      case 84: /* select ::= WITH RECURSIVE wqlist selectnowith */
 {
-  Select *p = yymsp[0].minor.yy457;
+  Select *p = yymsp[0].minor.yy539;
   if( p ){
-    p->pWith = yymsp[-1].minor.yy59;
+    p->pWith = yymsp[-1].minor.yy131;
     parserDoubleLinkSelect(pParse, p);
   }else{
-    sqlite3WithDelete(pParse->db, yymsp[-1].minor.yy59);
+    sqlite3WithDelete(pParse->db, yymsp[-1].minor.yy131);
   }
-  yymsp[-3].minor.yy457 = p;
+  yymsp[-3].minor.yy539 = p;
 }
         break;
-      case 82: /* select ::= selectnowith */
+      case 85: /* select ::= selectnowith */
 {
-  Select *p = yymsp[0].minor.yy457;
+  Select *p = yymsp[0].minor.yy539;
   if( p ){
     parserDoubleLinkSelect(pParse, p);
   }
-  yymsp[0].minor.yy457 = p; /*A-overwrites-X*/
+  yymsp[0].minor.yy539 = p; /*A-overwrites-X*/
 }
         break;
-      case 83: /* selectnowith ::= selectnowith multiselect_op oneselect */
+      case 86: /* selectnowith ::= selectnowith multiselect_op oneselect */
 {
-  Select *pRhs = yymsp[0].minor.yy457;
-  Select *pLhs = yymsp[-2].minor.yy457;
+  Select *pRhs = yymsp[0].minor.yy539;
+  Select *pLhs = yymsp[-2].minor.yy539;
   if( pRhs && pRhs->pPrior ){
     SrcList *pFrom;
     Token x;
@@ -151628,142 +155934,142 @@ static YYACTIONTYPE yy_reduce(
     pRhs = sqlite3SelectNew(pParse,0,pFrom,0,0,0,0,0,0);
   }
   if( pRhs ){
-    pRhs->op = (u8)yymsp[-1].minor.yy494;
+    pRhs->op = (u8)yymsp[-1].minor.yy192;
     pRhs->pPrior = pLhs;
     if( ALWAYS(pLhs) ) pLhs->selFlags &= ~SF_MultiValue;
     pRhs->selFlags &= ~SF_MultiValue;
-    if( yymsp[-1].minor.yy494!=TK_ALL ) pParse->hasCompound = 1;
+    if( yymsp[-1].minor.yy192!=TK_ALL ) pParse->hasCompound = 1;
   }else{
     sqlite3SelectDelete(pParse->db, pLhs);
   }
-  yymsp[-2].minor.yy457 = pRhs;
+  yymsp[-2].minor.yy539 = pRhs;
 }
         break;
-      case 84: /* multiselect_op ::= UNION */
-      case 86: /* multiselect_op ::= EXCEPT|INTERSECT */ yytestcase(yyruleno==86);
-{yymsp[0].minor.yy494 = yymsp[0].major; /*A-overwrites-OP*/}
+      case 87: /* multiselect_op ::= UNION */
+      case 89: /* multiselect_op ::= EXCEPT|INTERSECT */ yytestcase(yyruleno==89);
+{yymsp[0].minor.yy192 = yymsp[0].major; /*A-overwrites-OP*/}
         break;
-      case 85: /* multiselect_op ::= UNION ALL */
-{yymsp[-1].minor.yy494 = TK_ALL;}
+      case 88: /* multiselect_op ::= UNION ALL */
+{yymsp[-1].minor.yy192 = TK_ALL;}
         break;
-      case 87: /* oneselect ::= SELECT distinct selcollist from where_opt groupby_opt having_opt orderby_opt limit_opt */
+      case 90: /* oneselect ::= SELECT distinct selcollist from where_opt groupby_opt having_opt orderby_opt limit_opt */
 {
-  yymsp[-8].minor.yy457 = sqlite3SelectNew(pParse,yymsp[-6].minor.yy434,yymsp[-5].minor.yy483,yymsp[-4].minor.yy524,yymsp[-3].minor.yy434,yymsp[-2].minor.yy524,yymsp[-1].minor.yy434,yymsp[-7].minor.yy494,yymsp[0].minor.yy524);
+  yymsp[-8].minor.yy539 = sqlite3SelectNew(pParse,yymsp[-6].minor.yy242,yymsp[-5].minor.yy47,yymsp[-4].minor.yy202,yymsp[-3].minor.yy242,yymsp[-2].minor.yy202,yymsp[-1].minor.yy242,yymsp[-7].minor.yy192,yymsp[0].minor.yy202);
 }
         break;
-      case 88: /* oneselect ::= SELECT distinct selcollist from where_opt groupby_opt having_opt window_clause orderby_opt limit_opt */
+      case 91: /* oneselect ::= SELECT distinct selcollist from where_opt groupby_opt having_opt window_clause orderby_opt limit_opt */
 {
-  yymsp[-9].minor.yy457 = sqlite3SelectNew(pParse,yymsp[-7].minor.yy434,yymsp[-6].minor.yy483,yymsp[-5].minor.yy524,yymsp[-4].minor.yy434,yymsp[-3].minor.yy524,yymsp[-1].minor.yy434,yymsp[-8].minor.yy494,yymsp[0].minor.yy524);
-  if( yymsp[-9].minor.yy457 ){
-    yymsp[-9].minor.yy457->pWinDefn = yymsp[-2].minor.yy295;
+  yymsp[-9].minor.yy539 = sqlite3SelectNew(pParse,yymsp[-7].minor.yy242,yymsp[-6].minor.yy47,yymsp[-5].minor.yy202,yymsp[-4].minor.yy242,yymsp[-3].minor.yy202,yymsp[-1].minor.yy242,yymsp[-8].minor.yy192,yymsp[0].minor.yy202);
+  if( yymsp[-9].minor.yy539 ){
+    yymsp[-9].minor.yy539->pWinDefn = yymsp[-2].minor.yy303;
   }else{
-    sqlite3WindowListDelete(pParse->db, yymsp[-2].minor.yy295);
+    sqlite3WindowListDelete(pParse->db, yymsp[-2].minor.yy303);
   }
 }
         break;
-      case 89: /* values ::= VALUES LP nexprlist RP */
+      case 92: /* values ::= VALUES LP nexprlist RP */
 {
-  yymsp[-3].minor.yy457 = sqlite3SelectNew(pParse,yymsp[-1].minor.yy434,0,0,0,0,0,SF_Values,0);
+  yymsp[-3].minor.yy539 = sqlite3SelectNew(pParse,yymsp[-1].minor.yy242,0,0,0,0,0,SF_Values,0);
 }
         break;
-      case 90: /* values ::= values COMMA LP nexprlist RP */
+      case 93: /* values ::= values COMMA LP nexprlist RP */
 {
-  Select *pRight, *pLeft = yymsp[-4].minor.yy457;
-  pRight = sqlite3SelectNew(pParse,yymsp[-1].minor.yy434,0,0,0,0,0,SF_Values|SF_MultiValue,0);
+  Select *pRight, *pLeft = yymsp[-4].minor.yy539;
+  pRight = sqlite3SelectNew(pParse,yymsp[-1].minor.yy242,0,0,0,0,0,SF_Values|SF_MultiValue,0);
   if( ALWAYS(pLeft) ) pLeft->selFlags &= ~SF_MultiValue;
   if( pRight ){
     pRight->op = TK_ALL;
     pRight->pPrior = pLeft;
-    yymsp[-4].minor.yy457 = pRight;
+    yymsp[-4].minor.yy539 = pRight;
   }else{
-    yymsp[-4].minor.yy457 = pLeft;
+    yymsp[-4].minor.yy539 = pLeft;
   }
 }
         break;
-      case 91: /* distinct ::= DISTINCT */
-{yymsp[0].minor.yy494 = SF_Distinct;}
+      case 94: /* distinct ::= DISTINCT */
+{yymsp[0].minor.yy192 = SF_Distinct;}
         break;
-      case 92: /* distinct ::= ALL */
-{yymsp[0].minor.yy494 = SF_All;}
+      case 95: /* distinct ::= ALL */
+{yymsp[0].minor.yy192 = SF_All;}
         break;
-      case 94: /* sclp ::= */
-      case 127: /* orderby_opt ::= */ yytestcase(yyruleno==127);
-      case 134: /* groupby_opt ::= */ yytestcase(yyruleno==134);
-      case 214: /* exprlist ::= */ yytestcase(yyruleno==214);
-      case 217: /* paren_exprlist ::= */ yytestcase(yyruleno==217);
-      case 222: /* eidlist_opt ::= */ yytestcase(yyruleno==222);
-{yymsp[1].minor.yy434 = 0;}
+      case 97: /* sclp ::= */
+      case 130: /* orderby_opt ::= */ yytestcase(yyruleno==130);
+      case 140: /* groupby_opt ::= */ yytestcase(yyruleno==140);
+      case 220: /* exprlist ::= */ yytestcase(yyruleno==220);
+      case 223: /* paren_exprlist ::= */ yytestcase(yyruleno==223);
+      case 228: /* eidlist_opt ::= */ yytestcase(yyruleno==228);
+{yymsp[1].minor.yy242 = 0;}
         break;
-      case 95: /* selcollist ::= sclp scanpt expr scanpt as */
+      case 98: /* selcollist ::= sclp scanpt expr scanpt as */
 {
-   yymsp[-4].minor.yy434 = sqlite3ExprListAppend(pParse, yymsp[-4].minor.yy434, yymsp[-2].minor.yy524);
-   if( yymsp[0].minor.yy0.n>0 ) sqlite3ExprListSetName(pParse, yymsp[-4].minor.yy434, &yymsp[0].minor.yy0, 1);
-   sqlite3ExprListSetSpan(pParse,yymsp[-4].minor.yy434,yymsp[-3].minor.yy294,yymsp[-1].minor.yy294);
+   yymsp[-4].minor.yy242 = sqlite3ExprListAppend(pParse, yymsp[-4].minor.yy242, yymsp[-2].minor.yy202);
+   if( yymsp[0].minor.yy0.n>0 ) sqlite3ExprListSetName(pParse, yymsp[-4].minor.yy242, &yymsp[0].minor.yy0, 1);
+   sqlite3ExprListSetSpan(pParse,yymsp[-4].minor.yy242,yymsp[-3].minor.yy436,yymsp[-1].minor.yy436);
 }
         break;
-      case 96: /* selcollist ::= sclp scanpt STAR */
+      case 99: /* selcollist ::= sclp scanpt STAR */
 {
   Expr *p = sqlite3Expr(pParse->db, TK_ASTERISK, 0);
-  yymsp[-2].minor.yy434 = sqlite3ExprListAppend(pParse, yymsp[-2].minor.yy434, p);
+  yymsp[-2].minor.yy242 = sqlite3ExprListAppend(pParse, yymsp[-2].minor.yy242, p);
 }
         break;
-      case 97: /* selcollist ::= sclp scanpt nm DOT STAR */
+      case 100: /* selcollist ::= sclp scanpt nm DOT STAR */
 {
   Expr *pRight = sqlite3PExpr(pParse, TK_ASTERISK, 0, 0);
   Expr *pLeft = sqlite3ExprAlloc(pParse->db, TK_ID, &yymsp[-2].minor.yy0, 1);
   Expr *pDot = sqlite3PExpr(pParse, TK_DOT, pLeft, pRight);
-  yymsp[-4].minor.yy434 = sqlite3ExprListAppend(pParse,yymsp[-4].minor.yy434, pDot);
+  yymsp[-4].minor.yy242 = sqlite3ExprListAppend(pParse,yymsp[-4].minor.yy242, pDot);
 }
         break;
-      case 98: /* as ::= AS nm */
-      case 109: /* dbnm ::= DOT nm */ yytestcase(yyruleno==109);
-      case 238: /* plus_num ::= PLUS INTEGER|FLOAT */ yytestcase(yyruleno==238);
-      case 239: /* minus_num ::= MINUS INTEGER|FLOAT */ yytestcase(yyruleno==239);
+      case 101: /* as ::= AS nm */
+      case 112: /* dbnm ::= DOT nm */ yytestcase(yyruleno==112);
+      case 244: /* plus_num ::= PLUS INTEGER|FLOAT */ yytestcase(yyruleno==244);
+      case 245: /* minus_num ::= MINUS INTEGER|FLOAT */ yytestcase(yyruleno==245);
 {yymsp[-1].minor.yy0 = yymsp[0].minor.yy0;}
         break;
-      case 100: /* from ::= */
-{yymsp[1].minor.yy483 = sqlite3DbMallocZero(pParse->db, sizeof(*yymsp[1].minor.yy483));}
+      case 103: /* from ::= */
+{yymsp[1].minor.yy47 = sqlite3DbMallocZero(pParse->db, sizeof(*yymsp[1].minor.yy47));}
         break;
-      case 101: /* from ::= FROM seltablist */
+      case 104: /* from ::= FROM seltablist */
 {
-  yymsp[-1].minor.yy483 = yymsp[0].minor.yy483;
-  sqlite3SrcListShiftJoinType(yymsp[-1].minor.yy483);
+  yymsp[-1].minor.yy47 = yymsp[0].minor.yy47;
+  sqlite3SrcListShiftJoinType(yymsp[-1].minor.yy47);
 }
         break;
-      case 102: /* stl_prefix ::= seltablist joinop */
+      case 105: /* stl_prefix ::= seltablist joinop */
 {
-   if( ALWAYS(yymsp[-1].minor.yy483 && yymsp[-1].minor.yy483->nSrc>0) ) yymsp[-1].minor.yy483->a[yymsp[-1].minor.yy483->nSrc-1].fg.jointype = (u8)yymsp[0].minor.yy494;
+   if( ALWAYS(yymsp[-1].minor.yy47 && yymsp[-1].minor.yy47->nSrc>0) ) yymsp[-1].minor.yy47->a[yymsp[-1].minor.yy47->nSrc-1].fg.jointype = (u8)yymsp[0].minor.yy192;
 }
         break;
-      case 103: /* stl_prefix ::= */
-{yymsp[1].minor.yy483 = 0;}
+      case 106: /* stl_prefix ::= */
+{yymsp[1].minor.yy47 = 0;}
         break;
-      case 104: /* seltablist ::= stl_prefix nm dbnm as indexed_opt on_opt using_opt */
+      case 107: /* seltablist ::= stl_prefix nm dbnm as indexed_opt on_opt using_opt */
 {
-  yymsp[-6].minor.yy483 = sqlite3SrcListAppendFromTerm(pParse,yymsp[-6].minor.yy483,&yymsp[-5].minor.yy0,&yymsp[-4].minor.yy0,&yymsp[-3].minor.yy0,0,yymsp[-1].minor.yy524,yymsp[0].minor.yy62);
-  sqlite3SrcListIndexedBy(pParse, yymsp[-6].minor.yy483, &yymsp[-2].minor.yy0);
+  yymsp[-6].minor.yy47 = sqlite3SrcListAppendFromTerm(pParse,yymsp[-6].minor.yy47,&yymsp[-5].minor.yy0,&yymsp[-4].minor.yy0,&yymsp[-3].minor.yy0,0,yymsp[-1].minor.yy202,yymsp[0].minor.yy600);
+  sqlite3SrcListIndexedBy(pParse, yymsp[-6].minor.yy47, &yymsp[-2].minor.yy0);
 }
         break;
-      case 105: /* seltablist ::= stl_prefix nm dbnm LP exprlist RP as on_opt using_opt */
+      case 108: /* seltablist ::= stl_prefix nm dbnm LP exprlist RP as on_opt using_opt */
 {
-  yymsp[-8].minor.yy483 = sqlite3SrcListAppendFromTerm(pParse,yymsp[-8].minor.yy483,&yymsp[-7].minor.yy0,&yymsp[-6].minor.yy0,&yymsp[-2].minor.yy0,0,yymsp[-1].minor.yy524,yymsp[0].minor.yy62);
-  sqlite3SrcListFuncArgs(pParse, yymsp[-8].minor.yy483, yymsp[-4].minor.yy434);
+  yymsp[-8].minor.yy47 = sqlite3SrcListAppendFromTerm(pParse,yymsp[-8].minor.yy47,&yymsp[-7].minor.yy0,&yymsp[-6].minor.yy0,&yymsp[-2].minor.yy0,0,yymsp[-1].minor.yy202,yymsp[0].minor.yy600);
+  sqlite3SrcListFuncArgs(pParse, yymsp[-8].minor.yy47, yymsp[-4].minor.yy242);
 }
         break;
-      case 106: /* seltablist ::= stl_prefix LP select RP as on_opt using_opt */
+      case 109: /* seltablist ::= stl_prefix LP select RP as on_opt using_opt */
 {
-    yymsp[-6].minor.yy483 = sqlite3SrcListAppendFromTerm(pParse,yymsp[-6].minor.yy483,0,0,&yymsp[-2].minor.yy0,yymsp[-4].minor.yy457,yymsp[-1].minor.yy524,yymsp[0].minor.yy62);
+    yymsp[-6].minor.yy47 = sqlite3SrcListAppendFromTerm(pParse,yymsp[-6].minor.yy47,0,0,&yymsp[-2].minor.yy0,yymsp[-4].minor.yy539,yymsp[-1].minor.yy202,yymsp[0].minor.yy600);
   }
         break;
-      case 107: /* seltablist ::= stl_prefix LP seltablist RP as on_opt using_opt */
+      case 110: /* seltablist ::= stl_prefix LP seltablist RP as on_opt using_opt */
 {
-    if( yymsp[-6].minor.yy483==0 && yymsp[-2].minor.yy0.n==0 && yymsp[-1].minor.yy524==0 && yymsp[0].minor.yy62==0 ){
-      yymsp[-6].minor.yy483 = yymsp[-4].minor.yy483;
-    }else if( yymsp[-4].minor.yy483->nSrc==1 ){
-      yymsp[-6].minor.yy483 = sqlite3SrcListAppendFromTerm(pParse,yymsp[-6].minor.yy483,0,0,&yymsp[-2].minor.yy0,0,yymsp[-1].minor.yy524,yymsp[0].minor.yy62);
-      if( yymsp[-6].minor.yy483 ){
-        struct SrcList_item *pNew = &yymsp[-6].minor.yy483->a[yymsp[-6].minor.yy483->nSrc-1];
-        struct SrcList_item *pOld = yymsp[-4].minor.yy483->a;
+    if( yymsp[-6].minor.yy47==0 && yymsp[-2].minor.yy0.n==0 && yymsp[-1].minor.yy202==0 && yymsp[0].minor.yy600==0 ){
+      yymsp[-6].minor.yy47 = yymsp[-4].minor.yy47;
+    }else if( yymsp[-4].minor.yy47->nSrc==1 ){
+      yymsp[-6].minor.yy47 = sqlite3SrcListAppendFromTerm(pParse,yymsp[-6].minor.yy47,0,0,&yymsp[-2].minor.yy0,0,yymsp[-1].minor.yy202,yymsp[0].minor.yy600);
+      if( yymsp[-6].minor.yy47 ){
+        struct SrcList_item *pNew = &yymsp[-6].minor.yy47->a[yymsp[-6].minor.yy47->nSrc-1];
+        struct SrcList_item *pOld = yymsp[-4].minor.yy47->a;
         pNew->zName = pOld->zName;
         pNew->zDatabase = pOld->zDatabase;
         pNew->pSelect = pOld->pSelect;
@@ -151776,201 +156082,208 @@ static YYACTIONTYPE yy_reduce(
         pOld->zName = pOld->zDatabase = 0;
         pOld->pSelect = 0;
       }
-      sqlite3SrcListDelete(pParse->db, yymsp[-4].minor.yy483);
+      sqlite3SrcListDelete(pParse->db, yymsp[-4].minor.yy47);
     }else{
       Select *pSubquery;
-      sqlite3SrcListShiftJoinType(yymsp[-4].minor.yy483);
-      pSubquery = sqlite3SelectNew(pParse,0,yymsp[-4].minor.yy483,0,0,0,0,SF_NestedFrom,0);
-      yymsp[-6].minor.yy483 = sqlite3SrcListAppendFromTerm(pParse,yymsp[-6].minor.yy483,0,0,&yymsp[-2].minor.yy0,pSubquery,yymsp[-1].minor.yy524,yymsp[0].minor.yy62);
+      sqlite3SrcListShiftJoinType(yymsp[-4].minor.yy47);
+      pSubquery = sqlite3SelectNew(pParse,0,yymsp[-4].minor.yy47,0,0,0,0,SF_NestedFrom,0);
+      yymsp[-6].minor.yy47 = sqlite3SrcListAppendFromTerm(pParse,yymsp[-6].minor.yy47,0,0,&yymsp[-2].minor.yy0,pSubquery,yymsp[-1].minor.yy202,yymsp[0].minor.yy600);
     }
   }
         break;
-      case 108: /* dbnm ::= */
-      case 122: /* indexed_opt ::= */ yytestcase(yyruleno==122);
+      case 111: /* dbnm ::= */
+      case 125: /* indexed_opt ::= */ yytestcase(yyruleno==125);
 {yymsp[1].minor.yy0.z=0; yymsp[1].minor.yy0.n=0;}
         break;
-      case 110: /* fullname ::= nm */
+      case 113: /* fullname ::= nm */
 {
-  yylhsminor.yy483 = sqlite3SrcListAppend(pParse,0,&yymsp[0].minor.yy0,0);
-  if( IN_RENAME_OBJECT && yylhsminor.yy483 ) sqlite3RenameTokenMap(pParse, yylhsminor.yy483->a[0].zName, &yymsp[0].minor.yy0);
+  yylhsminor.yy47 = sqlite3SrcListAppend(pParse,0,&yymsp[0].minor.yy0,0);
+  if( IN_RENAME_OBJECT && yylhsminor.yy47 ) sqlite3RenameTokenMap(pParse, yylhsminor.yy47->a[0].zName, &yymsp[0].minor.yy0);
 }
-  yymsp[0].minor.yy483 = yylhsminor.yy483;
+  yymsp[0].minor.yy47 = yylhsminor.yy47;
         break;
-      case 111: /* fullname ::= nm DOT nm */
+      case 114: /* fullname ::= nm DOT nm */
 {
-  yylhsminor.yy483 = sqlite3SrcListAppend(pParse,0,&yymsp[-2].minor.yy0,&yymsp[0].minor.yy0);
-  if( IN_RENAME_OBJECT && yylhsminor.yy483 ) sqlite3RenameTokenMap(pParse, yylhsminor.yy483->a[0].zName, &yymsp[0].minor.yy0);
+  yylhsminor.yy47 = sqlite3SrcListAppend(pParse,0,&yymsp[-2].minor.yy0,&yymsp[0].minor.yy0);
+  if( IN_RENAME_OBJECT && yylhsminor.yy47 ) sqlite3RenameTokenMap(pParse, yylhsminor.yy47->a[0].zName, &yymsp[0].minor.yy0);
 }
-  yymsp[-2].minor.yy483 = yylhsminor.yy483;
+  yymsp[-2].minor.yy47 = yylhsminor.yy47;
         break;
-      case 112: /* xfullname ::= nm */
-{yymsp[0].minor.yy483 = sqlite3SrcListAppend(pParse,0,&yymsp[0].minor.yy0,0); /*A-overwrites-X*/}
+      case 115: /* xfullname ::= nm */
+{yymsp[0].minor.yy47 = sqlite3SrcListAppend(pParse,0,&yymsp[0].minor.yy0,0); /*A-overwrites-X*/}
         break;
-      case 113: /* xfullname ::= nm DOT nm */
-{yymsp[-2].minor.yy483 = sqlite3SrcListAppend(pParse,0,&yymsp[-2].minor.yy0,&yymsp[0].minor.yy0); /*A-overwrites-X*/}
+      case 116: /* xfullname ::= nm DOT nm */
+{yymsp[-2].minor.yy47 = sqlite3SrcListAppend(pParse,0,&yymsp[-2].minor.yy0,&yymsp[0].minor.yy0); /*A-overwrites-X*/}
         break;
-      case 114: /* xfullname ::= nm DOT nm AS nm */
+      case 117: /* xfullname ::= nm DOT nm AS nm */
 {
-   yymsp[-4].minor.yy483 = sqlite3SrcListAppend(pParse,0,&yymsp[-4].minor.yy0,&yymsp[-2].minor.yy0); /*A-overwrites-X*/
-   if( yymsp[-4].minor.yy483 ) yymsp[-4].minor.yy483->a[0].zAlias = sqlite3NameFromToken(pParse->db, &yymsp[0].minor.yy0);
+   yymsp[-4].minor.yy47 = sqlite3SrcListAppend(pParse,0,&yymsp[-4].minor.yy0,&yymsp[-2].minor.yy0); /*A-overwrites-X*/
+   if( yymsp[-4].minor.yy47 ) yymsp[-4].minor.yy47->a[0].zAlias = sqlite3NameFromToken(pParse->db, &yymsp[0].minor.yy0);
 }
         break;
-      case 115: /* xfullname ::= nm AS nm */
+      case 118: /* xfullname ::= nm AS nm */
 {  
-   yymsp[-2].minor.yy483 = sqlite3SrcListAppend(pParse,0,&yymsp[-2].minor.yy0,0); /*A-overwrites-X*/
-   if( yymsp[-2].minor.yy483 ) yymsp[-2].minor.yy483->a[0].zAlias = sqlite3NameFromToken(pParse->db, &yymsp[0].minor.yy0);
+   yymsp[-2].minor.yy47 = sqlite3SrcListAppend(pParse,0,&yymsp[-2].minor.yy0,0); /*A-overwrites-X*/
+   if( yymsp[-2].minor.yy47 ) yymsp[-2].minor.yy47->a[0].zAlias = sqlite3NameFromToken(pParse->db, &yymsp[0].minor.yy0);
 }
         break;
-      case 116: /* joinop ::= COMMA|JOIN */
-{ yymsp[0].minor.yy494 = JT_INNER; }
+      case 119: /* joinop ::= COMMA|JOIN */
+{ yymsp[0].minor.yy192 = JT_INNER; }
         break;
-      case 117: /* joinop ::= JOIN_KW JOIN */
-{yymsp[-1].minor.yy494 = sqlite3JoinType(pParse,&yymsp[-1].minor.yy0,0,0);  /*X-overwrites-A*/}
+      case 120: /* joinop ::= JOIN_KW JOIN */
+{yymsp[-1].minor.yy192 = sqlite3JoinType(pParse,&yymsp[-1].minor.yy0,0,0);  /*X-overwrites-A*/}
         break;
-      case 118: /* joinop ::= JOIN_KW nm JOIN */
-{yymsp[-2].minor.yy494 = sqlite3JoinType(pParse,&yymsp[-2].minor.yy0,&yymsp[-1].minor.yy0,0); /*X-overwrites-A*/}
+      case 121: /* joinop ::= JOIN_KW nm JOIN */
+{yymsp[-2].minor.yy192 = sqlite3JoinType(pParse,&yymsp[-2].minor.yy0,&yymsp[-1].minor.yy0,0); /*X-overwrites-A*/}
         break;
-      case 119: /* joinop ::= JOIN_KW nm nm JOIN */
-{yymsp[-3].minor.yy494 = sqlite3JoinType(pParse,&yymsp[-3].minor.yy0,&yymsp[-2].minor.yy0,&yymsp[-1].minor.yy0);/*X-overwrites-A*/}
+      case 122: /* joinop ::= JOIN_KW nm nm JOIN */
+{yymsp[-3].minor.yy192 = sqlite3JoinType(pParse,&yymsp[-3].minor.yy0,&yymsp[-2].minor.yy0,&yymsp[-1].minor.yy0);/*X-overwrites-A*/}
         break;
-      case 120: /* on_opt ::= ON expr */
-      case 137: /* having_opt ::= HAVING expr */ yytestcase(yyruleno==137);
-      case 144: /* where_opt ::= WHERE expr */ yytestcase(yyruleno==144);
-      case 210: /* case_else ::= ELSE expr */ yytestcase(yyruleno==210);
-      case 231: /* vinto ::= INTO expr */ yytestcase(yyruleno==231);
-{yymsp[-1].minor.yy524 = yymsp[0].minor.yy524;}
+      case 123: /* on_opt ::= ON expr */
+      case 143: /* having_opt ::= HAVING expr */ yytestcase(yyruleno==143);
+      case 150: /* where_opt ::= WHERE expr */ yytestcase(yyruleno==150);
+      case 216: /* case_else ::= ELSE expr */ yytestcase(yyruleno==216);
+      case 237: /* vinto ::= INTO expr */ yytestcase(yyruleno==237);
+{yymsp[-1].minor.yy202 = yymsp[0].minor.yy202;}
         break;
-      case 121: /* on_opt ::= */
-      case 136: /* having_opt ::= */ yytestcase(yyruleno==136);
-      case 138: /* limit_opt ::= */ yytestcase(yyruleno==138);
-      case 143: /* where_opt ::= */ yytestcase(yyruleno==143);
-      case 211: /* case_else ::= */ yytestcase(yyruleno==211);
-      case 213: /* case_operand ::= */ yytestcase(yyruleno==213);
-      case 232: /* vinto ::= */ yytestcase(yyruleno==232);
-{yymsp[1].minor.yy524 = 0;}
+      case 124: /* on_opt ::= */
+      case 142: /* having_opt ::= */ yytestcase(yyruleno==142);
+      case 144: /* limit_opt ::= */ yytestcase(yyruleno==144);
+      case 149: /* where_opt ::= */ yytestcase(yyruleno==149);
+      case 217: /* case_else ::= */ yytestcase(yyruleno==217);
+      case 219: /* case_operand ::= */ yytestcase(yyruleno==219);
+      case 238: /* vinto ::= */ yytestcase(yyruleno==238);
+{yymsp[1].minor.yy202 = 0;}
         break;
-      case 123: /* indexed_opt ::= INDEXED BY nm */
+      case 126: /* indexed_opt ::= INDEXED BY nm */
 {yymsp[-2].minor.yy0 = yymsp[0].minor.yy0;}
         break;
-      case 124: /* indexed_opt ::= NOT INDEXED */
+      case 127: /* indexed_opt ::= NOT INDEXED */
 {yymsp[-1].minor.yy0.z=0; yymsp[-1].minor.yy0.n=1;}
         break;
-      case 125: /* using_opt ::= USING LP idlist RP */
-{yymsp[-3].minor.yy62 = yymsp[-1].minor.yy62;}
+      case 128: /* using_opt ::= USING LP idlist RP */
+{yymsp[-3].minor.yy600 = yymsp[-1].minor.yy600;}
         break;
-      case 126: /* using_opt ::= */
-      case 158: /* idlist_opt ::= */ yytestcase(yyruleno==158);
-{yymsp[1].minor.yy62 = 0;}
+      case 129: /* using_opt ::= */
+      case 164: /* idlist_opt ::= */ yytestcase(yyruleno==164);
+{yymsp[1].minor.yy600 = 0;}
         break;
-      case 128: /* orderby_opt ::= ORDER BY sortlist */
-      case 135: /* groupby_opt ::= GROUP BY nexprlist */ yytestcase(yyruleno==135);
-{yymsp[-2].minor.yy434 = yymsp[0].minor.yy434;}
+      case 131: /* orderby_opt ::= ORDER BY sortlist */
+      case 141: /* groupby_opt ::= GROUP BY nexprlist */ yytestcase(yyruleno==141);
+{yymsp[-2].minor.yy242 = yymsp[0].minor.yy242;}
         break;
-      case 129: /* sortlist ::= sortlist COMMA expr sortorder */
+      case 132: /* sortlist ::= sortlist COMMA expr sortorder nulls */
 {
-  yymsp[-3].minor.yy434 = sqlite3ExprListAppend(pParse,yymsp[-3].minor.yy434,yymsp[-1].minor.yy524);
-  sqlite3ExprListSetSortOrder(yymsp[-3].minor.yy434,yymsp[0].minor.yy494);
+  yymsp[-4].minor.yy242 = sqlite3ExprListAppend(pParse,yymsp[-4].minor.yy242,yymsp[-2].minor.yy202);
+  sqlite3ExprListSetSortOrder(yymsp[-4].minor.yy242,yymsp[-1].minor.yy192,yymsp[0].minor.yy192);
 }
         break;
-      case 130: /* sortlist ::= expr sortorder */
+      case 133: /* sortlist ::= expr sortorder nulls */
 {
-  yymsp[-1].minor.yy434 = sqlite3ExprListAppend(pParse,0,yymsp[-1].minor.yy524); /*A-overwrites-Y*/
-  sqlite3ExprListSetSortOrder(yymsp[-1].minor.yy434,yymsp[0].minor.yy494);
+  yymsp[-2].minor.yy242 = sqlite3ExprListAppend(pParse,0,yymsp[-2].minor.yy202); /*A-overwrites-Y*/
+  sqlite3ExprListSetSortOrder(yymsp[-2].minor.yy242,yymsp[-1].minor.yy192,yymsp[0].minor.yy192);
 }
         break;
-      case 131: /* sortorder ::= ASC */
-{yymsp[0].minor.yy494 = SQLITE_SO_ASC;}
+      case 134: /* sortorder ::= ASC */
+{yymsp[0].minor.yy192 = SQLITE_SO_ASC;}
         break;
-      case 132: /* sortorder ::= DESC */
-{yymsp[0].minor.yy494 = SQLITE_SO_DESC;}
+      case 135: /* sortorder ::= DESC */
+{yymsp[0].minor.yy192 = SQLITE_SO_DESC;}
         break;
-      case 133: /* sortorder ::= */
-{yymsp[1].minor.yy494 = SQLITE_SO_UNDEFINED;}
+      case 136: /* sortorder ::= */
+      case 139: /* nulls ::= */ yytestcase(yyruleno==139);
+{yymsp[1].minor.yy192 = SQLITE_SO_UNDEFINED;}
         break;
-      case 139: /* limit_opt ::= LIMIT expr */
-{yymsp[-1].minor.yy524 = sqlite3PExpr(pParse,TK_LIMIT,yymsp[0].minor.yy524,0);}
+      case 137: /* nulls ::= NULLS FIRST */
+{yymsp[-1].minor.yy192 = SQLITE_SO_ASC;}
         break;
-      case 140: /* limit_opt ::= LIMIT expr OFFSET expr */
-{yymsp[-3].minor.yy524 = sqlite3PExpr(pParse,TK_LIMIT,yymsp[-2].minor.yy524,yymsp[0].minor.yy524);}
+      case 138: /* nulls ::= NULLS LAST */
+{yymsp[-1].minor.yy192 = SQLITE_SO_DESC;}
         break;
-      case 141: /* limit_opt ::= LIMIT expr COMMA expr */
-{yymsp[-3].minor.yy524 = sqlite3PExpr(pParse,TK_LIMIT,yymsp[0].minor.yy524,yymsp[-2].minor.yy524);}
+      case 145: /* limit_opt ::= LIMIT expr */
+{yymsp[-1].minor.yy202 = sqlite3PExpr(pParse,TK_LIMIT,yymsp[0].minor.yy202,0);}
         break;
-      case 142: /* cmd ::= with DELETE FROM xfullname indexed_opt where_opt */
+      case 146: /* limit_opt ::= LIMIT expr OFFSET expr */
+{yymsp[-3].minor.yy202 = sqlite3PExpr(pParse,TK_LIMIT,yymsp[-2].minor.yy202,yymsp[0].minor.yy202);}
+        break;
+      case 147: /* limit_opt ::= LIMIT expr COMMA expr */
+{yymsp[-3].minor.yy202 = sqlite3PExpr(pParse,TK_LIMIT,yymsp[0].minor.yy202,yymsp[-2].minor.yy202);}
+        break;
+      case 148: /* cmd ::= with DELETE FROM xfullname indexed_opt where_opt */
 {
-  sqlite3SrcListIndexedBy(pParse, yymsp[-2].minor.yy483, &yymsp[-1].minor.yy0);
-  sqlite3DeleteFrom(pParse,yymsp[-2].minor.yy483,yymsp[0].minor.yy524,0,0);
+  sqlite3SrcListIndexedBy(pParse, yymsp[-2].minor.yy47, &yymsp[-1].minor.yy0);
+  sqlite3DeleteFrom(pParse,yymsp[-2].minor.yy47,yymsp[0].minor.yy202,0,0);
 }
         break;
-      case 145: /* cmd ::= with UPDATE orconf xfullname indexed_opt SET setlist where_opt */
+      case 151: /* cmd ::= with UPDATE orconf xfullname indexed_opt SET setlist where_opt */
 {
-  sqlite3SrcListIndexedBy(pParse, yymsp[-4].minor.yy483, &yymsp[-3].minor.yy0);
-  sqlite3ExprListCheckLength(pParse,yymsp[-1].minor.yy434,"set list"); 
-  sqlite3Update(pParse,yymsp[-4].minor.yy483,yymsp[-1].minor.yy434,yymsp[0].minor.yy524,yymsp[-5].minor.yy494,0,0,0);
+  sqlite3SrcListIndexedBy(pParse, yymsp[-4].minor.yy47, &yymsp[-3].minor.yy0);
+  sqlite3ExprListCheckLength(pParse,yymsp[-1].minor.yy242,"set list"); 
+  sqlite3Update(pParse,yymsp[-4].minor.yy47,yymsp[-1].minor.yy242,yymsp[0].minor.yy202,yymsp[-5].minor.yy192,0,0,0);
 }
         break;
-      case 146: /* setlist ::= setlist COMMA nm EQ expr */
+      case 152: /* setlist ::= setlist COMMA nm EQ expr */
 {
-  yymsp[-4].minor.yy434 = sqlite3ExprListAppend(pParse, yymsp[-4].minor.yy434, yymsp[0].minor.yy524);
-  sqlite3ExprListSetName(pParse, yymsp[-4].minor.yy434, &yymsp[-2].minor.yy0, 1);
+  yymsp[-4].minor.yy242 = sqlite3ExprListAppend(pParse, yymsp[-4].minor.yy242, yymsp[0].minor.yy202);
+  sqlite3ExprListSetName(pParse, yymsp[-4].minor.yy242, &yymsp[-2].minor.yy0, 1);
 }
         break;
-      case 147: /* setlist ::= setlist COMMA LP idlist RP EQ expr */
+      case 153: /* setlist ::= setlist COMMA LP idlist RP EQ expr */
 {
-  yymsp[-6].minor.yy434 = sqlite3ExprListAppendVector(pParse, yymsp[-6].minor.yy434, yymsp[-3].minor.yy62, yymsp[0].minor.yy524);
+  yymsp[-6].minor.yy242 = sqlite3ExprListAppendVector(pParse, yymsp[-6].minor.yy242, yymsp[-3].minor.yy600, yymsp[0].minor.yy202);
 }
         break;
-      case 148: /* setlist ::= nm EQ expr */
+      case 154: /* setlist ::= nm EQ expr */
 {
-  yylhsminor.yy434 = sqlite3ExprListAppend(pParse, 0, yymsp[0].minor.yy524);
-  sqlite3ExprListSetName(pParse, yylhsminor.yy434, &yymsp[-2].minor.yy0, 1);
+  yylhsminor.yy242 = sqlite3ExprListAppend(pParse, 0, yymsp[0].minor.yy202);
+  sqlite3ExprListSetName(pParse, yylhsminor.yy242, &yymsp[-2].minor.yy0, 1);
 }
-  yymsp[-2].minor.yy434 = yylhsminor.yy434;
+  yymsp[-2].minor.yy242 = yylhsminor.yy242;
         break;
-      case 149: /* setlist ::= LP idlist RP EQ expr */
+      case 155: /* setlist ::= LP idlist RP EQ expr */
 {
-  yymsp[-4].minor.yy434 = sqlite3ExprListAppendVector(pParse, 0, yymsp[-3].minor.yy62, yymsp[0].minor.yy524);
+  yymsp[-4].minor.yy242 = sqlite3ExprListAppendVector(pParse, 0, yymsp[-3].minor.yy600, yymsp[0].minor.yy202);
 }
         break;
-      case 150: /* cmd ::= with insert_cmd INTO xfullname idlist_opt select upsert */
+      case 156: /* cmd ::= with insert_cmd INTO xfullname idlist_opt select upsert */
 {
-  sqlite3Insert(pParse, yymsp[-3].minor.yy483, yymsp[-1].minor.yy457, yymsp[-2].minor.yy62, yymsp[-5].minor.yy494, yymsp[0].minor.yy136);
+  sqlite3Insert(pParse, yymsp[-3].minor.yy47, yymsp[-1].minor.yy539, yymsp[-2].minor.yy600, yymsp[-5].minor.yy192, yymsp[0].minor.yy318);
 }
         break;
-      case 151: /* cmd ::= with insert_cmd INTO xfullname idlist_opt DEFAULT VALUES */
+      case 157: /* cmd ::= with insert_cmd INTO xfullname idlist_opt DEFAULT VALUES */
 {
-  sqlite3Insert(pParse, yymsp[-3].minor.yy483, 0, yymsp[-2].minor.yy62, yymsp[-5].minor.yy494, 0);
+  sqlite3Insert(pParse, yymsp[-3].minor.yy47, 0, yymsp[-2].minor.yy600, yymsp[-5].minor.yy192, 0);
 }
         break;
-      case 152: /* upsert ::= */
-{ yymsp[1].minor.yy136 = 0; }
+      case 158: /* upsert ::= */
+{ yymsp[1].minor.yy318 = 0; }
         break;
-      case 153: /* upsert ::= ON CONFLICT LP sortlist RP where_opt DO UPDATE SET setlist where_opt */
-{ yymsp[-10].minor.yy136 = sqlite3UpsertNew(pParse->db,yymsp[-7].minor.yy434,yymsp[-5].minor.yy524,yymsp[-1].minor.yy434,yymsp[0].minor.yy524);}
+      case 159: /* upsert ::= ON CONFLICT LP sortlist RP where_opt DO UPDATE SET setlist where_opt */
+{ yymsp[-10].minor.yy318 = sqlite3UpsertNew(pParse->db,yymsp[-7].minor.yy242,yymsp[-5].minor.yy202,yymsp[-1].minor.yy242,yymsp[0].minor.yy202);}
         break;
-      case 154: /* upsert ::= ON CONFLICT LP sortlist RP where_opt DO NOTHING */
-{ yymsp[-7].minor.yy136 = sqlite3UpsertNew(pParse->db,yymsp[-4].minor.yy434,yymsp[-2].minor.yy524,0,0); }
+      case 160: /* upsert ::= ON CONFLICT LP sortlist RP where_opt DO NOTHING */
+{ yymsp[-7].minor.yy318 = sqlite3UpsertNew(pParse->db,yymsp[-4].minor.yy242,yymsp[-2].minor.yy202,0,0); }
         break;
-      case 155: /* upsert ::= ON CONFLICT DO NOTHING */
-{ yymsp[-3].minor.yy136 = sqlite3UpsertNew(pParse->db,0,0,0,0); }
+      case 161: /* upsert ::= ON CONFLICT DO NOTHING */
+{ yymsp[-3].minor.yy318 = sqlite3UpsertNew(pParse->db,0,0,0,0); }
         break;
-      case 159: /* idlist_opt ::= LP idlist RP */
-{yymsp[-2].minor.yy62 = yymsp[-1].minor.yy62;}
+      case 165: /* idlist_opt ::= LP idlist RP */
+{yymsp[-2].minor.yy600 = yymsp[-1].minor.yy600;}
         break;
-      case 160: /* idlist ::= idlist COMMA nm */
-{yymsp[-2].minor.yy62 = sqlite3IdListAppend(pParse,yymsp[-2].minor.yy62,&yymsp[0].minor.yy0);}
+      case 166: /* idlist ::= idlist COMMA nm */
+{yymsp[-2].minor.yy600 = sqlite3IdListAppend(pParse,yymsp[-2].minor.yy600,&yymsp[0].minor.yy0);}
         break;
-      case 161: /* idlist ::= nm */
-{yymsp[0].minor.yy62 = sqlite3IdListAppend(pParse,0,&yymsp[0].minor.yy0); /*A-overwrites-Y*/}
+      case 167: /* idlist ::= nm */
+{yymsp[0].minor.yy600 = sqlite3IdListAppend(pParse,0,&yymsp[0].minor.yy0); /*A-overwrites-Y*/}
         break;
-      case 162: /* expr ::= LP expr RP */
-{yymsp[-2].minor.yy524 = yymsp[-1].minor.yy524;}
+      case 168: /* expr ::= LP expr RP */
+{yymsp[-2].minor.yy202 = yymsp[-1].minor.yy202;}
         break;
-      case 163: /* expr ::= ID|INDEXED */
-      case 164: /* expr ::= JOIN_KW */ yytestcase(yyruleno==164);
-{yymsp[0].minor.yy524=tokenExpr(pParse,TK_ID,yymsp[0].minor.yy0); /*A-overwrites-X*/}
+      case 169: /* expr ::= ID|INDEXED */
+      case 170: /* expr ::= JOIN_KW */ yytestcase(yyruleno==170);
+{yymsp[0].minor.yy202=tokenExpr(pParse,TK_ID,yymsp[0].minor.yy0); /*A-overwrites-X*/}
         break;
-      case 165: /* expr ::= nm DOT nm */
+      case 171: /* expr ::= nm DOT nm */
 {
   Expr *temp1 = sqlite3ExprAlloc(pParse->db, TK_ID, &yymsp[-2].minor.yy0, 1);
   Expr *temp2 = sqlite3ExprAlloc(pParse->db, TK_ID, &yymsp[0].minor.yy0, 1);
@@ -151978,11 +156291,11 @@ static YYACTIONTYPE yy_reduce(
     sqlite3RenameTokenMap(pParse, (void*)temp2, &yymsp[0].minor.yy0);
     sqlite3RenameTokenMap(pParse, (void*)temp1, &yymsp[-2].minor.yy0);
   }
-  yylhsminor.yy524 = sqlite3PExpr(pParse, TK_DOT, temp1, temp2);
+  yylhsminor.yy202 = sqlite3PExpr(pParse, TK_DOT, temp1, temp2);
 }
-  yymsp[-2].minor.yy524 = yylhsminor.yy524;
+  yymsp[-2].minor.yy202 = yylhsminor.yy202;
         break;
-      case 166: /* expr ::= nm DOT nm DOT nm */
+      case 172: /* expr ::= nm DOT nm DOT nm */
 {
   Expr *temp1 = sqlite3ExprAlloc(pParse->db, TK_ID, &yymsp[-4].minor.yy0, 1);
   Expr *temp2 = sqlite3ExprAlloc(pParse->db, TK_ID, &yymsp[-2].minor.yy0, 1);
@@ -151992,26 +156305,26 @@ static YYACTIONTYPE yy_reduce(
     sqlite3RenameTokenMap(pParse, (void*)temp3, &yymsp[0].minor.yy0);
     sqlite3RenameTokenMap(pParse, (void*)temp2, &yymsp[-2].minor.yy0);
   }
-  yylhsminor.yy524 = sqlite3PExpr(pParse, TK_DOT, temp1, temp4);
+  yylhsminor.yy202 = sqlite3PExpr(pParse, TK_DOT, temp1, temp4);
 }
-  yymsp[-4].minor.yy524 = yylhsminor.yy524;
+  yymsp[-4].minor.yy202 = yylhsminor.yy202;
         break;
-      case 167: /* term ::= NULL|FLOAT|BLOB */
-      case 168: /* term ::= STRING */ yytestcase(yyruleno==168);
-{yymsp[0].minor.yy524=tokenExpr(pParse,yymsp[0].major,yymsp[0].minor.yy0); /*A-overwrites-X*/}
+      case 173: /* term ::= NULL|FLOAT|BLOB */
+      case 174: /* term ::= STRING */ yytestcase(yyruleno==174);
+{yymsp[0].minor.yy202=tokenExpr(pParse,yymsp[0].major,yymsp[0].minor.yy0); /*A-overwrites-X*/}
         break;
-      case 169: /* term ::= INTEGER */
+      case 175: /* term ::= INTEGER */
 {
-  yylhsminor.yy524 = sqlite3ExprAlloc(pParse->db, TK_INTEGER, &yymsp[0].minor.yy0, 1);
+  yylhsminor.yy202 = sqlite3ExprAlloc(pParse->db, TK_INTEGER, &yymsp[0].minor.yy0, 1);
 }
-  yymsp[0].minor.yy524 = yylhsminor.yy524;
+  yymsp[0].minor.yy202 = yylhsminor.yy202;
         break;
-      case 170: /* expr ::= VARIABLE */
+      case 176: /* expr ::= VARIABLE */
 {
   if( !(yymsp[0].minor.yy0.z[0]=='#' && sqlite3Isdigit(yymsp[0].minor.yy0.z[1])) ){
     u32 n = yymsp[0].minor.yy0.n;
-    yymsp[0].minor.yy524 = tokenExpr(pParse, TK_VARIABLE, yymsp[0].minor.yy0);
-    sqlite3ExprAssignVarNumber(pParse, yymsp[0].minor.yy524, n);
+    yymsp[0].minor.yy202 = tokenExpr(pParse, TK_VARIABLE, yymsp[0].minor.yy0);
+    sqlite3ExprAssignVarNumber(pParse, yymsp[0].minor.yy202, n);
   }else{
     /* When doing a nested parse, one can include terms in an expression
     ** that look like this:   #1 #2 ...  These terms refer to registers
@@ -152020,154 +156333,159 @@ static YYACTIONTYPE yy_reduce(
     assert( t.n>=2 );
     if( pParse->nested==0 ){
       sqlite3ErrorMsg(pParse, "near \"%T\": syntax error", &t);
-      yymsp[0].minor.yy524 = 0;
+      yymsp[0].minor.yy202 = 0;
     }else{
-      yymsp[0].minor.yy524 = sqlite3PExpr(pParse, TK_REGISTER, 0, 0);
-      if( yymsp[0].minor.yy524 ) sqlite3GetInt32(&t.z[1], &yymsp[0].minor.yy524->iTable);
+      yymsp[0].minor.yy202 = sqlite3PExpr(pParse, TK_REGISTER, 0, 0);
+      if( yymsp[0].minor.yy202 ) sqlite3GetInt32(&t.z[1], &yymsp[0].minor.yy202->iTable);
     }
   }
 }
         break;
-      case 171: /* expr ::= expr COLLATE ID|STRING */
+      case 177: /* expr ::= expr COLLATE ID|STRING */
 {
-  yymsp[-2].minor.yy524 = sqlite3ExprAddCollateToken(pParse, yymsp[-2].minor.yy524, &yymsp[0].minor.yy0, 1);
+  yymsp[-2].minor.yy202 = sqlite3ExprAddCollateToken(pParse, yymsp[-2].minor.yy202, &yymsp[0].minor.yy0, 1);
 }
         break;
-      case 172: /* expr ::= CAST LP expr AS typetoken RP */
+      case 178: /* expr ::= CAST LP expr AS typetoken RP */
 {
-  yymsp[-5].minor.yy524 = sqlite3ExprAlloc(pParse->db, TK_CAST, &yymsp[-1].minor.yy0, 1);
-  sqlite3ExprAttachSubtrees(pParse->db, yymsp[-5].minor.yy524, yymsp[-3].minor.yy524, 0);
+  yymsp[-5].minor.yy202 = sqlite3ExprAlloc(pParse->db, TK_CAST, &yymsp[-1].minor.yy0, 1);
+  sqlite3ExprAttachSubtrees(pParse->db, yymsp[-5].minor.yy202, yymsp[-3].minor.yy202, 0);
 }
         break;
-      case 173: /* expr ::= ID|INDEXED LP distinct exprlist RP */
+      case 179: /* expr ::= ID|INDEXED LP distinct exprlist RP */
 {
-  yylhsminor.yy524 = sqlite3ExprFunction(pParse, yymsp[-1].minor.yy434, &yymsp[-4].minor.yy0, yymsp[-2].minor.yy494);
+  yylhsminor.yy202 = sqlite3ExprFunction(pParse, yymsp[-1].minor.yy242, &yymsp[-4].minor.yy0, yymsp[-2].minor.yy192);
 }
-  yymsp[-4].minor.yy524 = yylhsminor.yy524;
+  yymsp[-4].minor.yy202 = yylhsminor.yy202;
         break;
-      case 174: /* expr ::= ID|INDEXED LP STAR RP */
+      case 180: /* expr ::= ID|INDEXED LP STAR RP */
 {
-  yylhsminor.yy524 = sqlite3ExprFunction(pParse, 0, &yymsp[-3].minor.yy0, 0);
+  yylhsminor.yy202 = sqlite3ExprFunction(pParse, 0, &yymsp[-3].minor.yy0, 0);
 }
-  yymsp[-3].minor.yy524 = yylhsminor.yy524;
+  yymsp[-3].minor.yy202 = yylhsminor.yy202;
         break;
-      case 175: /* expr ::= ID|INDEXED LP distinct exprlist RP over_clause */
+      case 181: /* expr ::= ID|INDEXED LP distinct exprlist RP filter_over */
 {
-  yylhsminor.yy524 = sqlite3ExprFunction(pParse, yymsp[-2].minor.yy434, &yymsp[-5].minor.yy0, yymsp[-3].minor.yy494);
-  sqlite3WindowAttach(pParse, yylhsminor.yy524, yymsp[0].minor.yy295);
+  yylhsminor.yy202 = sqlite3ExprFunction(pParse, yymsp[-2].minor.yy242, &yymsp[-5].minor.yy0, yymsp[-3].minor.yy192);
+  sqlite3WindowAttach(pParse, yylhsminor.yy202, yymsp[0].minor.yy303);
 }
-  yymsp[-5].minor.yy524 = yylhsminor.yy524;
+  yymsp[-5].minor.yy202 = yylhsminor.yy202;
         break;
-      case 176: /* expr ::= ID|INDEXED LP STAR RP over_clause */
+      case 182: /* expr ::= ID|INDEXED LP STAR RP filter_over */
 {
-  yylhsminor.yy524 = sqlite3ExprFunction(pParse, 0, &yymsp[-4].minor.yy0, 0);
-  sqlite3WindowAttach(pParse, yylhsminor.yy524, yymsp[0].minor.yy295);
+  yylhsminor.yy202 = sqlite3ExprFunction(pParse, 0, &yymsp[-4].minor.yy0, 0);
+  sqlite3WindowAttach(pParse, yylhsminor.yy202, yymsp[0].minor.yy303);
 }
-  yymsp[-4].minor.yy524 = yylhsminor.yy524;
+  yymsp[-4].minor.yy202 = yylhsminor.yy202;
         break;
-      case 177: /* term ::= CTIME_KW */
+      case 183: /* term ::= CTIME_KW */
 {
-  yylhsminor.yy524 = sqlite3ExprFunction(pParse, 0, &yymsp[0].minor.yy0, 0);
+  yylhsminor.yy202 = sqlite3ExprFunction(pParse, 0, &yymsp[0].minor.yy0, 0);
 }
-  yymsp[0].minor.yy524 = yylhsminor.yy524;
+  yymsp[0].minor.yy202 = yylhsminor.yy202;
         break;
-      case 178: /* expr ::= LP nexprlist COMMA expr RP */
+      case 184: /* expr ::= LP nexprlist COMMA expr RP */
 {
-  ExprList *pList = sqlite3ExprListAppend(pParse, yymsp[-3].minor.yy434, yymsp[-1].minor.yy524);
-  yymsp[-4].minor.yy524 = sqlite3PExpr(pParse, TK_VECTOR, 0, 0);
-  if( yymsp[-4].minor.yy524 ){
-    yymsp[-4].minor.yy524->x.pList = pList;
+  ExprList *pList = sqlite3ExprListAppend(pParse, yymsp[-3].minor.yy242, yymsp[-1].minor.yy202);
+  yymsp[-4].minor.yy202 = sqlite3PExpr(pParse, TK_VECTOR, 0, 0);
+  if( yymsp[-4].minor.yy202 ){
+    yymsp[-4].minor.yy202->x.pList = pList;
+    if( ALWAYS(pList->nExpr) ){
+      yymsp[-4].minor.yy202->flags |= pList->a[0].pExpr->flags & EP_Propagate;
+    }
   }else{
     sqlite3ExprListDelete(pParse->db, pList);
   }
 }
         break;
-      case 179: /* expr ::= expr AND expr */
-      case 180: /* expr ::= expr OR expr */ yytestcase(yyruleno==180);
-      case 181: /* expr ::= expr LT|GT|GE|LE expr */ yytestcase(yyruleno==181);
-      case 182: /* expr ::= expr EQ|NE expr */ yytestcase(yyruleno==182);
-      case 183: /* expr ::= expr BITAND|BITOR|LSHIFT|RSHIFT expr */ yytestcase(yyruleno==183);
-      case 184: /* expr ::= expr PLUS|MINUS expr */ yytestcase(yyruleno==184);
-      case 185: /* expr ::= expr STAR|SLASH|REM expr */ yytestcase(yyruleno==185);
-      case 186: /* expr ::= expr CONCAT expr */ yytestcase(yyruleno==186);
-{yymsp[-2].minor.yy524=sqlite3PExpr(pParse,yymsp[-1].major,yymsp[-2].minor.yy524,yymsp[0].minor.yy524);}
+      case 185: /* expr ::= expr AND expr */
+{yymsp[-2].minor.yy202=sqlite3ExprAnd(pParse,yymsp[-2].minor.yy202,yymsp[0].minor.yy202);}
         break;
-      case 187: /* likeop ::= NOT LIKE_KW|MATCH */
+      case 186: /* expr ::= expr OR expr */
+      case 187: /* expr ::= expr LT|GT|GE|LE expr */ yytestcase(yyruleno==187);
+      case 188: /* expr ::= expr EQ|NE expr */ yytestcase(yyruleno==188);
+      case 189: /* expr ::= expr BITAND|BITOR|LSHIFT|RSHIFT expr */ yytestcase(yyruleno==189);
+      case 190: /* expr ::= expr PLUS|MINUS expr */ yytestcase(yyruleno==190);
+      case 191: /* expr ::= expr STAR|SLASH|REM expr */ yytestcase(yyruleno==191);
+      case 192: /* expr ::= expr CONCAT expr */ yytestcase(yyruleno==192);
+{yymsp[-2].minor.yy202=sqlite3PExpr(pParse,yymsp[-1].major,yymsp[-2].minor.yy202,yymsp[0].minor.yy202);}
+        break;
+      case 193: /* likeop ::= NOT LIKE_KW|MATCH */
 {yymsp[-1].minor.yy0=yymsp[0].minor.yy0; yymsp[-1].minor.yy0.n|=0x80000000; /*yymsp[-1].minor.yy0-overwrite-yymsp[0].minor.yy0*/}
         break;
-      case 188: /* expr ::= expr likeop expr */
+      case 194: /* expr ::= expr likeop expr */
 {
   ExprList *pList;
   int bNot = yymsp[-1].minor.yy0.n & 0x80000000;
   yymsp[-1].minor.yy0.n &= 0x7fffffff;
-  pList = sqlite3ExprListAppend(pParse,0, yymsp[0].minor.yy524);
-  pList = sqlite3ExprListAppend(pParse,pList, yymsp[-2].minor.yy524);
-  yymsp[-2].minor.yy524 = sqlite3ExprFunction(pParse, pList, &yymsp[-1].minor.yy0, 0);
-  if( bNot ) yymsp[-2].minor.yy524 = sqlite3PExpr(pParse, TK_NOT, yymsp[-2].minor.yy524, 0);
-  if( yymsp[-2].minor.yy524 ) yymsp[-2].minor.yy524->flags |= EP_InfixFunc;
+  pList = sqlite3ExprListAppend(pParse,0, yymsp[0].minor.yy202);
+  pList = sqlite3ExprListAppend(pParse,pList, yymsp[-2].minor.yy202);
+  yymsp[-2].minor.yy202 = sqlite3ExprFunction(pParse, pList, &yymsp[-1].minor.yy0, 0);
+  if( bNot ) yymsp[-2].minor.yy202 = sqlite3PExpr(pParse, TK_NOT, yymsp[-2].minor.yy202, 0);
+  if( yymsp[-2].minor.yy202 ) yymsp[-2].minor.yy202->flags |= EP_InfixFunc;
 }
         break;
-      case 189: /* expr ::= expr likeop expr ESCAPE expr */
+      case 195: /* expr ::= expr likeop expr ESCAPE expr */
 {
   ExprList *pList;
   int bNot = yymsp[-3].minor.yy0.n & 0x80000000;
   yymsp[-3].minor.yy0.n &= 0x7fffffff;
-  pList = sqlite3ExprListAppend(pParse,0, yymsp[-2].minor.yy524);
-  pList = sqlite3ExprListAppend(pParse,pList, yymsp[-4].minor.yy524);
-  pList = sqlite3ExprListAppend(pParse,pList, yymsp[0].minor.yy524);
-  yymsp[-4].minor.yy524 = sqlite3ExprFunction(pParse, pList, &yymsp[-3].minor.yy0, 0);
-  if( bNot ) yymsp[-4].minor.yy524 = sqlite3PExpr(pParse, TK_NOT, yymsp[-4].minor.yy524, 0);
-  if( yymsp[-4].minor.yy524 ) yymsp[-4].minor.yy524->flags |= EP_InfixFunc;
+  pList = sqlite3ExprListAppend(pParse,0, yymsp[-2].minor.yy202);
+  pList = sqlite3ExprListAppend(pParse,pList, yymsp[-4].minor.yy202);
+  pList = sqlite3ExprListAppend(pParse,pList, yymsp[0].minor.yy202);
+  yymsp[-4].minor.yy202 = sqlite3ExprFunction(pParse, pList, &yymsp[-3].minor.yy0, 0);
+  if( bNot ) yymsp[-4].minor.yy202 = sqlite3PExpr(pParse, TK_NOT, yymsp[-4].minor.yy202, 0);
+  if( yymsp[-4].minor.yy202 ) yymsp[-4].minor.yy202->flags |= EP_InfixFunc;
 }
         break;
-      case 190: /* expr ::= expr ISNULL|NOTNULL */
-{yymsp[-1].minor.yy524 = sqlite3PExpr(pParse,yymsp[0].major,yymsp[-1].minor.yy524,0);}
+      case 196: /* expr ::= expr ISNULL|NOTNULL */
+{yymsp[-1].minor.yy202 = sqlite3PExpr(pParse,yymsp[0].major,yymsp[-1].minor.yy202,0);}
         break;
-      case 191: /* expr ::= expr NOT NULL */
-{yymsp[-2].minor.yy524 = sqlite3PExpr(pParse,TK_NOTNULL,yymsp[-2].minor.yy524,0);}
+      case 197: /* expr ::= expr NOT NULL */
+{yymsp[-2].minor.yy202 = sqlite3PExpr(pParse,TK_NOTNULL,yymsp[-2].minor.yy202,0);}
         break;
-      case 192: /* expr ::= expr IS expr */
+      case 198: /* expr ::= expr IS expr */
 {
-  yymsp[-2].minor.yy524 = sqlite3PExpr(pParse,TK_IS,yymsp[-2].minor.yy524,yymsp[0].minor.yy524);
-  binaryToUnaryIfNull(pParse, yymsp[0].minor.yy524, yymsp[-2].minor.yy524, TK_ISNULL);
+  yymsp[-2].minor.yy202 = sqlite3PExpr(pParse,TK_IS,yymsp[-2].minor.yy202,yymsp[0].minor.yy202);
+  binaryToUnaryIfNull(pParse, yymsp[0].minor.yy202, yymsp[-2].minor.yy202, TK_ISNULL);
 }
         break;
-      case 193: /* expr ::= expr IS NOT expr */
+      case 199: /* expr ::= expr IS NOT expr */
 {
-  yymsp[-3].minor.yy524 = sqlite3PExpr(pParse,TK_ISNOT,yymsp[-3].minor.yy524,yymsp[0].minor.yy524);
-  binaryToUnaryIfNull(pParse, yymsp[0].minor.yy524, yymsp[-3].minor.yy524, TK_NOTNULL);
+  yymsp[-3].minor.yy202 = sqlite3PExpr(pParse,TK_ISNOT,yymsp[-3].minor.yy202,yymsp[0].minor.yy202);
+  binaryToUnaryIfNull(pParse, yymsp[0].minor.yy202, yymsp[-3].minor.yy202, TK_NOTNULL);
 }
         break;
-      case 194: /* expr ::= NOT expr */
-      case 195: /* expr ::= BITNOT expr */ yytestcase(yyruleno==195);
-{yymsp[-1].minor.yy524 = sqlite3PExpr(pParse, yymsp[-1].major, yymsp[0].minor.yy524, 0);/*A-overwrites-B*/}
+      case 200: /* expr ::= NOT expr */
+      case 201: /* expr ::= BITNOT expr */ yytestcase(yyruleno==201);
+{yymsp[-1].minor.yy202 = sqlite3PExpr(pParse, yymsp[-1].major, yymsp[0].minor.yy202, 0);/*A-overwrites-B*/}
         break;
-      case 196: /* expr ::= PLUS|MINUS expr */
+      case 202: /* expr ::= PLUS|MINUS expr */
 {
-  yymsp[-1].minor.yy524 = sqlite3PExpr(pParse, yymsp[-1].major==TK_PLUS ? TK_UPLUS : TK_UMINUS, yymsp[0].minor.yy524, 0);
+  yymsp[-1].minor.yy202 = sqlite3PExpr(pParse, yymsp[-1].major==TK_PLUS ? TK_UPLUS : TK_UMINUS, yymsp[0].minor.yy202, 0);
   /*A-overwrites-B*/
 }
         break;
-      case 197: /* between_op ::= BETWEEN */
-      case 200: /* in_op ::= IN */ yytestcase(yyruleno==200);
-{yymsp[0].minor.yy494 = 0;}
+      case 203: /* between_op ::= BETWEEN */
+      case 206: /* in_op ::= IN */ yytestcase(yyruleno==206);
+{yymsp[0].minor.yy192 = 0;}
         break;
-      case 199: /* expr ::= expr between_op expr AND expr */
+      case 205: /* expr ::= expr between_op expr AND expr */
 {
-  ExprList *pList = sqlite3ExprListAppend(pParse,0, yymsp[-2].minor.yy524);
-  pList = sqlite3ExprListAppend(pParse,pList, yymsp[0].minor.yy524);
-  yymsp[-4].minor.yy524 = sqlite3PExpr(pParse, TK_BETWEEN, yymsp[-4].minor.yy524, 0);
-  if( yymsp[-4].minor.yy524 ){
-    yymsp[-4].minor.yy524->x.pList = pList;
+  ExprList *pList = sqlite3ExprListAppend(pParse,0, yymsp[-2].minor.yy202);
+  pList = sqlite3ExprListAppend(pParse,pList, yymsp[0].minor.yy202);
+  yymsp[-4].minor.yy202 = sqlite3PExpr(pParse, TK_BETWEEN, yymsp[-4].minor.yy202, 0);
+  if( yymsp[-4].minor.yy202 ){
+    yymsp[-4].minor.yy202->x.pList = pList;
   }else{
     sqlite3ExprListDelete(pParse->db, pList);
   } 
-  if( yymsp[-3].minor.yy494 ) yymsp[-4].minor.yy524 = sqlite3PExpr(pParse, TK_NOT, yymsp[-4].minor.yy524, 0);
+  if( yymsp[-3].minor.yy192 ) yymsp[-4].minor.yy202 = sqlite3PExpr(pParse, TK_NOT, yymsp[-4].minor.yy202, 0);
 }
         break;
-      case 202: /* expr ::= expr in_op LP exprlist RP */
+      case 208: /* expr ::= expr in_op LP exprlist RP */
 {
-    if( yymsp[-1].minor.yy434==0 ){
+    if( yymsp[-1].minor.yy242==0 ){
       /* Expressions of the form
       **
       **      expr1 IN ()
@@ -152176,220 +156494,190 @@ static YYACTIONTYPE yy_reduce(
       ** simplify to constants 0 (false) and 1 (true), respectively,
       ** regardless of the value of expr1.
       */
-      if( IN_RENAME_OBJECT==0 ){
-        sqlite3ExprDelete(pParse->db, yymsp[-4].minor.yy524);
-        yymsp[-4].minor.yy524 = sqlite3ExprAlloc(pParse->db, TK_INTEGER,&sqlite3IntTokens[yymsp[-3].minor.yy494],1);
-      }
-    }else if( yymsp[-1].minor.yy434->nExpr==1 ){
-      /* Expressions of the form:
-      **
-      **      expr1 IN (?1)
-      **      expr1 NOT IN (?2)
-      **
-      ** with exactly one value on the RHS can be simplified to something
-      ** like this:
-      **
-      **      expr1 == ?1
-      **      expr1 <> ?2
-      **
-      ** But, the RHS of the == or <> is marked with the EP_Generic flag
-      ** so that it may not contribute to the computation of comparison
-      ** affinity or the collating sequence to use for comparison.  Otherwise,
-      ** the semantics would be subtly different from IN or NOT IN.
-      */
-      Expr *pRHS = yymsp[-1].minor.yy434->a[0].pExpr;
-      yymsp[-1].minor.yy434->a[0].pExpr = 0;
-      sqlite3ExprListDelete(pParse->db, yymsp[-1].minor.yy434);
-      /* pRHS cannot be NULL because a malloc error would have been detected
-      ** before now and control would have never reached this point */
-      if( ALWAYS(pRHS) ){
-        pRHS->flags &= ~EP_Collate;
-        pRHS->flags |= EP_Generic;
-      }
-      yymsp[-4].minor.yy524 = sqlite3PExpr(pParse, yymsp[-3].minor.yy494 ? TK_NE : TK_EQ, yymsp[-4].minor.yy524, pRHS);
-    }else{
-      yymsp[-4].minor.yy524 = sqlite3PExpr(pParse, TK_IN, yymsp[-4].minor.yy524, 0);
-      if( yymsp[-4].minor.yy524 ){
-        yymsp[-4].minor.yy524->x.pList = yymsp[-1].minor.yy434;
-        sqlite3ExprSetHeightAndFlags(pParse, yymsp[-4].minor.yy524);
+      sqlite3ExprUnmapAndDelete(pParse, yymsp[-4].minor.yy202);
+      yymsp[-4].minor.yy202 = sqlite3Expr(pParse->db, TK_INTEGER, yymsp[-3].minor.yy192 ? "1" : "0");
+    }else{
+      yymsp[-4].minor.yy202 = sqlite3PExpr(pParse, TK_IN, yymsp[-4].minor.yy202, 0);
+      if( yymsp[-4].minor.yy202 ){
+        yymsp[-4].minor.yy202->x.pList = yymsp[-1].minor.yy242;
+        sqlite3ExprSetHeightAndFlags(pParse, yymsp[-4].minor.yy202);
       }else{
-        sqlite3ExprListDelete(pParse->db, yymsp[-1].minor.yy434);
+        sqlite3ExprListDelete(pParse->db, yymsp[-1].minor.yy242);
       }
-      if( yymsp[-3].minor.yy494 ) yymsp[-4].minor.yy524 = sqlite3PExpr(pParse, TK_NOT, yymsp[-4].minor.yy524, 0);
+      if( yymsp[-3].minor.yy192 ) yymsp[-4].minor.yy202 = sqlite3PExpr(pParse, TK_NOT, yymsp[-4].minor.yy202, 0);
     }
   }
         break;
-      case 203: /* expr ::= LP select RP */
+      case 209: /* expr ::= LP select RP */
 {
-    yymsp[-2].minor.yy524 = sqlite3PExpr(pParse, TK_SELECT, 0, 0);
-    sqlite3PExprAddSelect(pParse, yymsp[-2].minor.yy524, yymsp[-1].minor.yy457);
+    yymsp[-2].minor.yy202 = sqlite3PExpr(pParse, TK_SELECT, 0, 0);
+    sqlite3PExprAddSelect(pParse, yymsp[-2].minor.yy202, yymsp[-1].minor.yy539);
   }
         break;
-      case 204: /* expr ::= expr in_op LP select RP */
+      case 210: /* expr ::= expr in_op LP select RP */
 {
-    yymsp[-4].minor.yy524 = sqlite3PExpr(pParse, TK_IN, yymsp[-4].minor.yy524, 0);
-    sqlite3PExprAddSelect(pParse, yymsp[-4].minor.yy524, yymsp[-1].minor.yy457);
-    if( yymsp[-3].minor.yy494 ) yymsp[-4].minor.yy524 = sqlite3PExpr(pParse, TK_NOT, yymsp[-4].minor.yy524, 0);
+    yymsp[-4].minor.yy202 = sqlite3PExpr(pParse, TK_IN, yymsp[-4].minor.yy202, 0);
+    sqlite3PExprAddSelect(pParse, yymsp[-4].minor.yy202, yymsp[-1].minor.yy539);
+    if( yymsp[-3].minor.yy192 ) yymsp[-4].minor.yy202 = sqlite3PExpr(pParse, TK_NOT, yymsp[-4].minor.yy202, 0);
   }
         break;
-      case 205: /* expr ::= expr in_op nm dbnm paren_exprlist */
+      case 211: /* expr ::= expr in_op nm dbnm paren_exprlist */
 {
     SrcList *pSrc = sqlite3SrcListAppend(pParse, 0,&yymsp[-2].minor.yy0,&yymsp[-1].minor.yy0);
     Select *pSelect = sqlite3SelectNew(pParse, 0,pSrc,0,0,0,0,0,0);
-    if( yymsp[0].minor.yy434 )  sqlite3SrcListFuncArgs(pParse, pSelect ? pSrc : 0, yymsp[0].minor.yy434);
-    yymsp[-4].minor.yy524 = sqlite3PExpr(pParse, TK_IN, yymsp[-4].minor.yy524, 0);
-    sqlite3PExprAddSelect(pParse, yymsp[-4].minor.yy524, pSelect);
-    if( yymsp[-3].minor.yy494 ) yymsp[-4].minor.yy524 = sqlite3PExpr(pParse, TK_NOT, yymsp[-4].minor.yy524, 0);
+    if( yymsp[0].minor.yy242 )  sqlite3SrcListFuncArgs(pParse, pSelect ? pSrc : 0, yymsp[0].minor.yy242);
+    yymsp[-4].minor.yy202 = sqlite3PExpr(pParse, TK_IN, yymsp[-4].minor.yy202, 0);
+    sqlite3PExprAddSelect(pParse, yymsp[-4].minor.yy202, pSelect);
+    if( yymsp[-3].minor.yy192 ) yymsp[-4].minor.yy202 = sqlite3PExpr(pParse, TK_NOT, yymsp[-4].minor.yy202, 0);
   }
         break;
-      case 206: /* expr ::= EXISTS LP select RP */
+      case 212: /* expr ::= EXISTS LP select RP */
 {
     Expr *p;
-    p = yymsp[-3].minor.yy524 = sqlite3PExpr(pParse, TK_EXISTS, 0, 0);
-    sqlite3PExprAddSelect(pParse, p, yymsp[-1].minor.yy457);
+    p = yymsp[-3].minor.yy202 = sqlite3PExpr(pParse, TK_EXISTS, 0, 0);
+    sqlite3PExprAddSelect(pParse, p, yymsp[-1].minor.yy539);
   }
         break;
-      case 207: /* expr ::= CASE case_operand case_exprlist case_else END */
+      case 213: /* expr ::= CASE case_operand case_exprlist case_else END */
 {
-  yymsp[-4].minor.yy524 = sqlite3PExpr(pParse, TK_CASE, yymsp[-3].minor.yy524, 0);
-  if( yymsp[-4].minor.yy524 ){
-    yymsp[-4].minor.yy524->x.pList = yymsp[-1].minor.yy524 ? sqlite3ExprListAppend(pParse,yymsp[-2].minor.yy434,yymsp[-1].minor.yy524) : yymsp[-2].minor.yy434;
-    sqlite3ExprSetHeightAndFlags(pParse, yymsp[-4].minor.yy524);
+  yymsp[-4].minor.yy202 = sqlite3PExpr(pParse, TK_CASE, yymsp[-3].minor.yy202, 0);
+  if( yymsp[-4].minor.yy202 ){
+    yymsp[-4].minor.yy202->x.pList = yymsp[-1].minor.yy202 ? sqlite3ExprListAppend(pParse,yymsp[-2].minor.yy242,yymsp[-1].minor.yy202) : yymsp[-2].minor.yy242;
+    sqlite3ExprSetHeightAndFlags(pParse, yymsp[-4].minor.yy202);
   }else{
-    sqlite3ExprListDelete(pParse->db, yymsp[-2].minor.yy434);
-    sqlite3ExprDelete(pParse->db, yymsp[-1].minor.yy524);
+    sqlite3ExprListDelete(pParse->db, yymsp[-2].minor.yy242);
+    sqlite3ExprDelete(pParse->db, yymsp[-1].minor.yy202);
   }
 }
         break;
-      case 208: /* case_exprlist ::= case_exprlist WHEN expr THEN expr */
+      case 214: /* case_exprlist ::= case_exprlist WHEN expr THEN expr */
 {
-  yymsp[-4].minor.yy434 = sqlite3ExprListAppend(pParse,yymsp[-4].minor.yy434, yymsp[-2].minor.yy524);
-  yymsp[-4].minor.yy434 = sqlite3ExprListAppend(pParse,yymsp[-4].minor.yy434, yymsp[0].minor.yy524);
+  yymsp[-4].minor.yy242 = sqlite3ExprListAppend(pParse,yymsp[-4].minor.yy242, yymsp[-2].minor.yy202);
+  yymsp[-4].minor.yy242 = sqlite3ExprListAppend(pParse,yymsp[-4].minor.yy242, yymsp[0].minor.yy202);
 }
         break;
-      case 209: /* case_exprlist ::= WHEN expr THEN expr */
+      case 215: /* case_exprlist ::= WHEN expr THEN expr */
 {
-  yymsp[-3].minor.yy434 = sqlite3ExprListAppend(pParse,0, yymsp[-2].minor.yy524);
-  yymsp[-3].minor.yy434 = sqlite3ExprListAppend(pParse,yymsp[-3].minor.yy434, yymsp[0].minor.yy524);
+  yymsp[-3].minor.yy242 = sqlite3ExprListAppend(pParse,0, yymsp[-2].minor.yy202);
+  yymsp[-3].minor.yy242 = sqlite3ExprListAppend(pParse,yymsp[-3].minor.yy242, yymsp[0].minor.yy202);
 }
         break;
-      case 212: /* case_operand ::= expr */
-{ /*A-overwrites-X*/}
+      case 218: /* case_operand ::= expr */
+{yymsp[0].minor.yy202 = yymsp[0].minor.yy202; /*A-overwrites-X*/}
         break;
-      case 215: /* nexprlist ::= nexprlist COMMA expr */
-{yymsp[-2].minor.yy434 = sqlite3ExprListAppend(pParse,yymsp[-2].minor.yy434,yymsp[0].minor.yy524);}
+      case 221: /* nexprlist ::= nexprlist COMMA expr */
+{yymsp[-2].minor.yy242 = sqlite3ExprListAppend(pParse,yymsp[-2].minor.yy242,yymsp[0].minor.yy202);}
         break;
-      case 216: /* nexprlist ::= expr */
-{yymsp[0].minor.yy434 = sqlite3ExprListAppend(pParse,0,yymsp[0].minor.yy524); /*A-overwrites-Y*/}
+      case 222: /* nexprlist ::= expr */
+{yymsp[0].minor.yy242 = sqlite3ExprListAppend(pParse,0,yymsp[0].minor.yy202); /*A-overwrites-Y*/}
         break;
-      case 218: /* paren_exprlist ::= LP exprlist RP */
-      case 223: /* eidlist_opt ::= LP eidlist RP */ yytestcase(yyruleno==223);
-{yymsp[-2].minor.yy434 = yymsp[-1].minor.yy434;}
+      case 224: /* paren_exprlist ::= LP exprlist RP */
+      case 229: /* eidlist_opt ::= LP eidlist RP */ yytestcase(yyruleno==229);
+{yymsp[-2].minor.yy242 = yymsp[-1].minor.yy242;}
         break;
-      case 219: /* cmd ::= createkw uniqueflag INDEX ifnotexists nm dbnm ON nm LP sortlist RP where_opt */
+      case 225: /* cmd ::= createkw uniqueflag INDEX ifnotexists nm dbnm ON nm LP sortlist RP where_opt */
 {
   sqlite3CreateIndex(pParse, &yymsp[-7].minor.yy0, &yymsp[-6].minor.yy0, 
-                     sqlite3SrcListAppend(pParse,0,&yymsp[-4].minor.yy0,0), yymsp[-2].minor.yy434, yymsp[-10].minor.yy494,
-                      &yymsp[-11].minor.yy0, yymsp[0].minor.yy524, SQLITE_SO_ASC, yymsp[-8].minor.yy494, SQLITE_IDXTYPE_APPDEF);
+                     sqlite3SrcListAppend(pParse,0,&yymsp[-4].minor.yy0,0), yymsp[-2].minor.yy242, yymsp[-10].minor.yy192,
+                      &yymsp[-11].minor.yy0, yymsp[0].minor.yy202, SQLITE_SO_ASC, yymsp[-8].minor.yy192, SQLITE_IDXTYPE_APPDEF);
   if( IN_RENAME_OBJECT && pParse->pNewIndex ){
     sqlite3RenameTokenMap(pParse, pParse->pNewIndex->zName, &yymsp[-4].minor.yy0);
   }
 }
         break;
-      case 220: /* uniqueflag ::= UNIQUE */
-      case 262: /* raisetype ::= ABORT */ yytestcase(yyruleno==262);
-{yymsp[0].minor.yy494 = OE_Abort;}
+      case 226: /* uniqueflag ::= UNIQUE */
+      case 268: /* raisetype ::= ABORT */ yytestcase(yyruleno==268);
+{yymsp[0].minor.yy192 = OE_Abort;}
         break;
-      case 221: /* uniqueflag ::= */
-{yymsp[1].minor.yy494 = OE_None;}
+      case 227: /* uniqueflag ::= */
+{yymsp[1].minor.yy192 = OE_None;}
         break;
-      case 224: /* eidlist ::= eidlist COMMA nm collate sortorder */
+      case 230: /* eidlist ::= eidlist COMMA nm collate sortorder */
 {
-  yymsp[-4].minor.yy434 = parserAddExprIdListTerm(pParse, yymsp[-4].minor.yy434, &yymsp[-2].minor.yy0, yymsp[-1].minor.yy494, yymsp[0].minor.yy494);
+  yymsp[-4].minor.yy242 = parserAddExprIdListTerm(pParse, yymsp[-4].minor.yy242, &yymsp[-2].minor.yy0, yymsp[-1].minor.yy192, yymsp[0].minor.yy192);
 }
         break;
-      case 225: /* eidlist ::= nm collate sortorder */
+      case 231: /* eidlist ::= nm collate sortorder */
 {
-  yymsp[-2].minor.yy434 = parserAddExprIdListTerm(pParse, 0, &yymsp[-2].minor.yy0, yymsp[-1].minor.yy494, yymsp[0].minor.yy494); /*A-overwrites-Y*/
+  yymsp[-2].minor.yy242 = parserAddExprIdListTerm(pParse, 0, &yymsp[-2].minor.yy0, yymsp[-1].minor.yy192, yymsp[0].minor.yy192); /*A-overwrites-Y*/
 }
         break;
-      case 228: /* cmd ::= DROP INDEX ifexists fullname */
-{sqlite3DropIndex(pParse, yymsp[0].minor.yy483, yymsp[-1].minor.yy494);}
+      case 234: /* cmd ::= DROP INDEX ifexists fullname */
+{sqlite3DropIndex(pParse, yymsp[0].minor.yy47, yymsp[-1].minor.yy192);}
         break;
-      case 229: /* cmd ::= VACUUM vinto */
-{sqlite3Vacuum(pParse,0,yymsp[0].minor.yy524);}
+      case 235: /* cmd ::= VACUUM vinto */
+{sqlite3Vacuum(pParse,0,yymsp[0].minor.yy202);}
         break;
-      case 230: /* cmd ::= VACUUM nm vinto */
-{sqlite3Vacuum(pParse,&yymsp[-1].minor.yy0,yymsp[0].minor.yy524);}
+      case 236: /* cmd ::= VACUUM nm vinto */
+{sqlite3Vacuum(pParse,&yymsp[-1].minor.yy0,yymsp[0].minor.yy202);}
         break;
-      case 233: /* cmd ::= PRAGMA nm dbnm */
+      case 239: /* cmd ::= PRAGMA nm dbnm */
 {sqlite3Pragma(pParse,&yymsp[-1].minor.yy0,&yymsp[0].minor.yy0,0,0);}
         break;
-      case 234: /* cmd ::= PRAGMA nm dbnm EQ nmnum */
+      case 240: /* cmd ::= PRAGMA nm dbnm EQ nmnum */
 {sqlite3Pragma(pParse,&yymsp[-3].minor.yy0,&yymsp[-2].minor.yy0,&yymsp[0].minor.yy0,0);}
         break;
-      case 235: /* cmd ::= PRAGMA nm dbnm LP nmnum RP */
+      case 241: /* cmd ::= PRAGMA nm dbnm LP nmnum RP */
 {sqlite3Pragma(pParse,&yymsp[-4].minor.yy0,&yymsp[-3].minor.yy0,&yymsp[-1].minor.yy0,0);}
         break;
-      case 236: /* cmd ::= PRAGMA nm dbnm EQ minus_num */
+      case 242: /* cmd ::= PRAGMA nm dbnm EQ minus_num */
 {sqlite3Pragma(pParse,&yymsp[-3].minor.yy0,&yymsp[-2].minor.yy0,&yymsp[0].minor.yy0,1);}
         break;
-      case 237: /* cmd ::= PRAGMA nm dbnm LP minus_num RP */
+      case 243: /* cmd ::= PRAGMA nm dbnm LP minus_num RP */
 {sqlite3Pragma(pParse,&yymsp[-4].minor.yy0,&yymsp[-3].minor.yy0,&yymsp[-1].minor.yy0,1);}
         break;
-      case 240: /* cmd ::= createkw trigger_decl BEGIN trigger_cmd_list END */
+      case 246: /* cmd ::= createkw trigger_decl BEGIN trigger_cmd_list END */
 {
   Token all;
   all.z = yymsp[-3].minor.yy0.z;
   all.n = (int)(yymsp[0].minor.yy0.z - yymsp[-3].minor.yy0.z) + yymsp[0].minor.yy0.n;
-  sqlite3FinishTrigger(pParse, yymsp[-1].minor.yy455, &all);
+  sqlite3FinishTrigger(pParse, yymsp[-1].minor.yy447, &all);
 }
         break;
-      case 241: /* trigger_decl ::= temp TRIGGER ifnotexists nm dbnm trigger_time trigger_event ON fullname foreach_clause when_clause */
+      case 247: /* trigger_decl ::= temp TRIGGER ifnotexists nm dbnm trigger_time trigger_event ON fullname foreach_clause when_clause */
 {
-  sqlite3BeginTrigger(pParse, &yymsp[-7].minor.yy0, &yymsp[-6].minor.yy0, yymsp[-5].minor.yy494, yymsp[-4].minor.yy90.a, yymsp[-4].minor.yy90.b, yymsp[-2].minor.yy483, yymsp[0].minor.yy524, yymsp[-10].minor.yy494, yymsp[-8].minor.yy494);
+  sqlite3BeginTrigger(pParse, &yymsp[-7].minor.yy0, &yymsp[-6].minor.yy0, yymsp[-5].minor.yy192, yymsp[-4].minor.yy230.a, yymsp[-4].minor.yy230.b, yymsp[-2].minor.yy47, yymsp[0].minor.yy202, yymsp[-10].minor.yy192, yymsp[-8].minor.yy192);
   yymsp[-10].minor.yy0 = (yymsp[-6].minor.yy0.n==0?yymsp[-7].minor.yy0:yymsp[-6].minor.yy0); /*A-overwrites-T*/
 }
         break;
-      case 242: /* trigger_time ::= BEFORE|AFTER */
-{ yymsp[0].minor.yy494 = yymsp[0].major; /*A-overwrites-X*/ }
+      case 248: /* trigger_time ::= BEFORE|AFTER */
+{ yymsp[0].minor.yy192 = yymsp[0].major; /*A-overwrites-X*/ }
         break;
-      case 243: /* trigger_time ::= INSTEAD OF */
-{ yymsp[-1].minor.yy494 = TK_INSTEAD;}
+      case 249: /* trigger_time ::= INSTEAD OF */
+{ yymsp[-1].minor.yy192 = TK_INSTEAD;}
         break;
-      case 244: /* trigger_time ::= */
-{ yymsp[1].minor.yy494 = TK_BEFORE; }
+      case 250: /* trigger_time ::= */
+{ yymsp[1].minor.yy192 = TK_BEFORE; }
         break;
-      case 245: /* trigger_event ::= DELETE|INSERT */
-      case 246: /* trigger_event ::= UPDATE */ yytestcase(yyruleno==246);
-{yymsp[0].minor.yy90.a = yymsp[0].major; /*A-overwrites-X*/ yymsp[0].minor.yy90.b = 0;}
+      case 251: /* trigger_event ::= DELETE|INSERT */
+      case 252: /* trigger_event ::= UPDATE */ yytestcase(yyruleno==252);
+{yymsp[0].minor.yy230.a = yymsp[0].major; /*A-overwrites-X*/ yymsp[0].minor.yy230.b = 0;}
         break;
-      case 247: /* trigger_event ::= UPDATE OF idlist */
-{yymsp[-2].minor.yy90.a = TK_UPDATE; yymsp[-2].minor.yy90.b = yymsp[0].minor.yy62;}
+      case 253: /* trigger_event ::= UPDATE OF idlist */
+{yymsp[-2].minor.yy230.a = TK_UPDATE; yymsp[-2].minor.yy230.b = yymsp[0].minor.yy600;}
         break;
-      case 248: /* when_clause ::= */
-      case 267: /* key_opt ::= */ yytestcase(yyruleno==267);
-      case 315: /* filter_opt ::= */ yytestcase(yyruleno==315);
-{ yymsp[1].minor.yy524 = 0; }
+      case 254: /* when_clause ::= */
+      case 273: /* key_opt ::= */ yytestcase(yyruleno==273);
+{ yymsp[1].minor.yy202 = 0; }
         break;
-      case 249: /* when_clause ::= WHEN expr */
-      case 268: /* key_opt ::= KEY expr */ yytestcase(yyruleno==268);
-{ yymsp[-1].minor.yy524 = yymsp[0].minor.yy524; }
+      case 255: /* when_clause ::= WHEN expr */
+      case 274: /* key_opt ::= KEY expr */ yytestcase(yyruleno==274);
+{ yymsp[-1].minor.yy202 = yymsp[0].minor.yy202; }
         break;
-      case 250: /* trigger_cmd_list ::= trigger_cmd_list trigger_cmd SEMI */
+      case 256: /* trigger_cmd_list ::= trigger_cmd_list trigger_cmd SEMI */
 {
-  assert( yymsp[-2].minor.yy455!=0 );
-  yymsp[-2].minor.yy455->pLast->pNext = yymsp[-1].minor.yy455;
-  yymsp[-2].minor.yy455->pLast = yymsp[-1].minor.yy455;
+  assert( yymsp[-2].minor.yy447!=0 );
+  yymsp[-2].minor.yy447->pLast->pNext = yymsp[-1].minor.yy447;
+  yymsp[-2].minor.yy447->pLast = yymsp[-1].minor.yy447;
 }
         break;
-      case 251: /* trigger_cmd_list ::= trigger_cmd SEMI */
+      case 257: /* trigger_cmd_list ::= trigger_cmd SEMI */
 { 
-  assert( yymsp[-1].minor.yy455!=0 );
-  yymsp[-1].minor.yy455->pLast = yymsp[-1].minor.yy455;
+  assert( yymsp[-1].minor.yy447!=0 );
+  yymsp[-1].minor.yy447->pLast = yymsp[-1].minor.yy447;
 }
         break;
-      case 252: /* trnm ::= nm DOT nm */
+      case 258: /* trnm ::= nm DOT nm */
 {
   yymsp[-2].minor.yy0 = yymsp[0].minor.yy0;
   sqlite3ErrorMsg(pParse, 
@@ -152397,328 +156685,344 @@ static YYACTIONTYPE yy_reduce(
         "statements within triggers");
 }
         break;
-      case 253: /* tridxby ::= INDEXED BY nm */
+      case 259: /* tridxby ::= INDEXED BY nm */
 {
   sqlite3ErrorMsg(pParse,
         "the INDEXED BY clause is not allowed on UPDATE or DELETE statements "
         "within triggers");
 }
         break;
-      case 254: /* tridxby ::= NOT INDEXED */
+      case 260: /* tridxby ::= NOT INDEXED */
 {
   sqlite3ErrorMsg(pParse,
         "the NOT INDEXED clause is not allowed on UPDATE or DELETE statements "
         "within triggers");
 }
         break;
-      case 255: /* trigger_cmd ::= UPDATE orconf trnm tridxby SET setlist where_opt scanpt */
-{yylhsminor.yy455 = sqlite3TriggerUpdateStep(pParse, &yymsp[-5].minor.yy0, yymsp[-2].minor.yy434, yymsp[-1].minor.yy524, yymsp[-6].minor.yy494, yymsp[-7].minor.yy0.z, yymsp[0].minor.yy294);}
-  yymsp[-7].minor.yy455 = yylhsminor.yy455;
+      case 261: /* trigger_cmd ::= UPDATE orconf trnm tridxby SET setlist where_opt scanpt */
+{yylhsminor.yy447 = sqlite3TriggerUpdateStep(pParse, &yymsp[-5].minor.yy0, yymsp[-2].minor.yy242, yymsp[-1].minor.yy202, yymsp[-6].minor.yy192, yymsp[-7].minor.yy0.z, yymsp[0].minor.yy436);}
+  yymsp[-7].minor.yy447 = yylhsminor.yy447;
         break;
-      case 256: /* trigger_cmd ::= scanpt insert_cmd INTO trnm idlist_opt select upsert scanpt */
+      case 262: /* trigger_cmd ::= scanpt insert_cmd INTO trnm idlist_opt select upsert scanpt */
 {
-   yylhsminor.yy455 = sqlite3TriggerInsertStep(pParse,&yymsp[-4].minor.yy0,yymsp[-3].minor.yy62,yymsp[-2].minor.yy457,yymsp[-6].minor.yy494,yymsp[-1].minor.yy136,yymsp[-7].minor.yy294,yymsp[0].minor.yy294);/*yylhsminor.yy455-overwrites-yymsp[-6].minor.yy494*/
+   yylhsminor.yy447 = sqlite3TriggerInsertStep(pParse,&yymsp[-4].minor.yy0,yymsp[-3].minor.yy600,yymsp[-2].minor.yy539,yymsp[-6].minor.yy192,yymsp[-1].minor.yy318,yymsp[-7].minor.yy436,yymsp[0].minor.yy436);/*yylhsminor.yy447-overwrites-yymsp[-6].minor.yy192*/
 }
-  yymsp[-7].minor.yy455 = yylhsminor.yy455;
+  yymsp[-7].minor.yy447 = yylhsminor.yy447;
         break;
-      case 257: /* trigger_cmd ::= DELETE FROM trnm tridxby where_opt scanpt */
-{yylhsminor.yy455 = sqlite3TriggerDeleteStep(pParse, &yymsp[-3].minor.yy0, yymsp[-1].minor.yy524, yymsp[-5].minor.yy0.z, yymsp[0].minor.yy294);}
-  yymsp[-5].minor.yy455 = yylhsminor.yy455;
+      case 263: /* trigger_cmd ::= DELETE FROM trnm tridxby where_opt scanpt */
+{yylhsminor.yy447 = sqlite3TriggerDeleteStep(pParse, &yymsp[-3].minor.yy0, yymsp[-1].minor.yy202, yymsp[-5].minor.yy0.z, yymsp[0].minor.yy436);}
+  yymsp[-5].minor.yy447 = yylhsminor.yy447;
         break;
-      case 258: /* trigger_cmd ::= scanpt select scanpt */
-{yylhsminor.yy455 = sqlite3TriggerSelectStep(pParse->db, yymsp[-1].minor.yy457, yymsp[-2].minor.yy294, yymsp[0].minor.yy294); /*yylhsminor.yy455-overwrites-yymsp[-1].minor.yy457*/}
-  yymsp[-2].minor.yy455 = yylhsminor.yy455;
+      case 264: /* trigger_cmd ::= scanpt select scanpt */
+{yylhsminor.yy447 = sqlite3TriggerSelectStep(pParse->db, yymsp[-1].minor.yy539, yymsp[-2].minor.yy436, yymsp[0].minor.yy436); /*yylhsminor.yy447-overwrites-yymsp[-1].minor.yy539*/}
+  yymsp[-2].minor.yy447 = yylhsminor.yy447;
         break;
-      case 259: /* expr ::= RAISE LP IGNORE RP */
+      case 265: /* expr ::= RAISE LP IGNORE RP */
 {
-  yymsp[-3].minor.yy524 = sqlite3PExpr(pParse, TK_RAISE, 0, 0); 
-  if( yymsp[-3].minor.yy524 ){
-    yymsp[-3].minor.yy524->affinity = OE_Ignore;
+  yymsp[-3].minor.yy202 = sqlite3PExpr(pParse, TK_RAISE, 0, 0); 
+  if( yymsp[-3].minor.yy202 ){
+    yymsp[-3].minor.yy202->affExpr = OE_Ignore;
   }
 }
         break;
-      case 260: /* expr ::= RAISE LP raisetype COMMA nm RP */
+      case 266: /* expr ::= RAISE LP raisetype COMMA nm RP */
 {
-  yymsp[-5].minor.yy524 = sqlite3ExprAlloc(pParse->db, TK_RAISE, &yymsp[-1].minor.yy0, 1); 
-  if( yymsp[-5].minor.yy524 ) {
-    yymsp[-5].minor.yy524->affinity = (char)yymsp[-3].minor.yy494;
+  yymsp[-5].minor.yy202 = sqlite3ExprAlloc(pParse->db, TK_RAISE, &yymsp[-1].minor.yy0, 1); 
+  if( yymsp[-5].minor.yy202 ) {
+    yymsp[-5].minor.yy202->affExpr = (char)yymsp[-3].minor.yy192;
   }
 }
         break;
-      case 261: /* raisetype ::= ROLLBACK */
-{yymsp[0].minor.yy494 = OE_Rollback;}
+      case 267: /* raisetype ::= ROLLBACK */
+{yymsp[0].minor.yy192 = OE_Rollback;}
         break;
-      case 263: /* raisetype ::= FAIL */
-{yymsp[0].minor.yy494 = OE_Fail;}
+      case 269: /* raisetype ::= FAIL */
+{yymsp[0].minor.yy192 = OE_Fail;}
         break;
-      case 264: /* cmd ::= DROP TRIGGER ifexists fullname */
+      case 270: /* cmd ::= DROP TRIGGER ifexists fullname */
 {
-  sqlite3DropTrigger(pParse,yymsp[0].minor.yy483,yymsp[-1].minor.yy494);
+  sqlite3DropTrigger(pParse,yymsp[0].minor.yy47,yymsp[-1].minor.yy192);
 }
         break;
-      case 265: /* cmd ::= ATTACH database_kw_opt expr AS expr key_opt */
+      case 271: /* cmd ::= ATTACH database_kw_opt expr AS expr key_opt */
 {
-  sqlite3Attach(pParse, yymsp[-3].minor.yy524, yymsp[-1].minor.yy524, yymsp[0].minor.yy524);
+  sqlite3Attach(pParse, yymsp[-3].minor.yy202, yymsp[-1].minor.yy202, yymsp[0].minor.yy202);
 }
         break;
-      case 266: /* cmd ::= DETACH database_kw_opt expr */
+      case 272: /* cmd ::= DETACH database_kw_opt expr */
 {
-  sqlite3Detach(pParse, yymsp[0].minor.yy524);
+  sqlite3Detach(pParse, yymsp[0].minor.yy202);
 }
         break;
-      case 269: /* cmd ::= REINDEX */
+      case 275: /* cmd ::= REINDEX */
 {sqlite3Reindex(pParse, 0, 0);}
         break;
-      case 270: /* cmd ::= REINDEX nm dbnm */
+      case 276: /* cmd ::= REINDEX nm dbnm */
 {sqlite3Reindex(pParse, &yymsp[-1].minor.yy0, &yymsp[0].minor.yy0);}
         break;
-      case 271: /* cmd ::= ANALYZE */
+      case 277: /* cmd ::= ANALYZE */
 {sqlite3Analyze(pParse, 0, 0);}
         break;
-      case 272: /* cmd ::= ANALYZE nm dbnm */
+      case 278: /* cmd ::= ANALYZE nm dbnm */
 {sqlite3Analyze(pParse, &yymsp[-1].minor.yy0, &yymsp[0].minor.yy0);}
         break;
-      case 273: /* cmd ::= ALTER TABLE fullname RENAME TO nm */
+      case 279: /* cmd ::= ALTER TABLE fullname RENAME TO nm */
 {
-  sqlite3AlterRenameTable(pParse,yymsp[-3].minor.yy483,&yymsp[0].minor.yy0);
+  sqlite3AlterRenameTable(pParse,yymsp[-3].minor.yy47,&yymsp[0].minor.yy0);
 }
         break;
-      case 274: /* cmd ::= ALTER TABLE add_column_fullname ADD kwcolumn_opt columnname carglist */
+      case 280: /* cmd ::= ALTER TABLE add_column_fullname ADD kwcolumn_opt columnname carglist */
 {
   yymsp[-1].minor.yy0.n = (int)(pParse->sLastToken.z-yymsp[-1].minor.yy0.z) + pParse->sLastToken.n;
   sqlite3AlterFinishAddColumn(pParse, &yymsp[-1].minor.yy0);
 }
         break;
-      case 275: /* add_column_fullname ::= fullname */
+      case 281: /* add_column_fullname ::= fullname */
 {
   disableLookaside(pParse);
-  sqlite3AlterBeginAddColumn(pParse, yymsp[0].minor.yy483);
+  sqlite3AlterBeginAddColumn(pParse, yymsp[0].minor.yy47);
 }
         break;
-      case 276: /* cmd ::= ALTER TABLE fullname RENAME kwcolumn_opt nm TO nm */
+      case 282: /* cmd ::= ALTER TABLE fullname RENAME kwcolumn_opt nm TO nm */
 {
-  sqlite3AlterRenameColumn(pParse, yymsp[-5].minor.yy483, &yymsp[-2].minor.yy0, &yymsp[0].minor.yy0);
+  sqlite3AlterRenameColumn(pParse, yymsp[-5].minor.yy47, &yymsp[-2].minor.yy0, &yymsp[0].minor.yy0);
 }
         break;
-      case 277: /* cmd ::= create_vtab */
+      case 283: /* cmd ::= create_vtab */
 {sqlite3VtabFinishParse(pParse,0);}
         break;
-      case 278: /* cmd ::= create_vtab LP vtabarglist RP */
+      case 284: /* cmd ::= create_vtab LP vtabarglist RP */
 {sqlite3VtabFinishParse(pParse,&yymsp[0].minor.yy0);}
         break;
-      case 279: /* create_vtab ::= createkw VIRTUAL TABLE ifnotexists nm dbnm USING nm */
+      case 285: /* create_vtab ::= createkw VIRTUAL TABLE ifnotexists nm dbnm USING nm */
 {
-    sqlite3VtabBeginParse(pParse, &yymsp[-3].minor.yy0, &yymsp[-2].minor.yy0, &yymsp[0].minor.yy0, yymsp[-4].minor.yy494);
+    sqlite3VtabBeginParse(pParse, &yymsp[-3].minor.yy0, &yymsp[-2].minor.yy0, &yymsp[0].minor.yy0, yymsp[-4].minor.yy192);
 }
         break;
-      case 280: /* vtabarg ::= */
+      case 286: /* vtabarg ::= */
 {sqlite3VtabArgInit(pParse);}
         break;
-      case 281: /* vtabargtoken ::= ANY */
-      case 282: /* vtabargtoken ::= lp anylist RP */ yytestcase(yyruleno==282);
-      case 283: /* lp ::= LP */ yytestcase(yyruleno==283);
+      case 287: /* vtabargtoken ::= ANY */
+      case 288: /* vtabargtoken ::= lp anylist RP */ yytestcase(yyruleno==288);
+      case 289: /* lp ::= LP */ yytestcase(yyruleno==289);
 {sqlite3VtabArgExtend(pParse,&yymsp[0].minor.yy0);}
         break;
-      case 284: /* with ::= WITH wqlist */
-      case 285: /* with ::= WITH RECURSIVE wqlist */ yytestcase(yyruleno==285);
-{ sqlite3WithPush(pParse, yymsp[0].minor.yy59, 1); }
+      case 290: /* with ::= WITH wqlist */
+      case 291: /* with ::= WITH RECURSIVE wqlist */ yytestcase(yyruleno==291);
+{ sqlite3WithPush(pParse, yymsp[0].minor.yy131, 1); }
         break;
-      case 286: /* wqlist ::= nm eidlist_opt AS LP select RP */
+      case 292: /* wqlist ::= nm eidlist_opt AS LP select RP */
 {
-  yymsp[-5].minor.yy59 = sqlite3WithAdd(pParse, 0, &yymsp[-5].minor.yy0, yymsp[-4].minor.yy434, yymsp[-1].minor.yy457); /*A-overwrites-X*/
+  yymsp[-5].minor.yy131 = sqlite3WithAdd(pParse, 0, &yymsp[-5].minor.yy0, yymsp[-4].minor.yy242, yymsp[-1].minor.yy539); /*A-overwrites-X*/
 }
         break;
-      case 287: /* wqlist ::= wqlist COMMA nm eidlist_opt AS LP select RP */
+      case 293: /* wqlist ::= wqlist COMMA nm eidlist_opt AS LP select RP */
 {
-  yymsp[-7].minor.yy59 = sqlite3WithAdd(pParse, yymsp[-7].minor.yy59, &yymsp[-5].minor.yy0, yymsp[-4].minor.yy434, yymsp[-1].minor.yy457);
+  yymsp[-7].minor.yy131 = sqlite3WithAdd(pParse, yymsp[-7].minor.yy131, &yymsp[-5].minor.yy0, yymsp[-4].minor.yy242, yymsp[-1].minor.yy539);
 }
         break;
-      case 288: /* windowdefn_list ::= windowdefn */
-{ yylhsminor.yy295 = yymsp[0].minor.yy295; }
-  yymsp[0].minor.yy295 = yylhsminor.yy295;
+      case 294: /* windowdefn_list ::= windowdefn */
+{ yylhsminor.yy303 = yymsp[0].minor.yy303; }
+  yymsp[0].minor.yy303 = yylhsminor.yy303;
         break;
-      case 289: /* windowdefn_list ::= windowdefn_list COMMA windowdefn */
+      case 295: /* windowdefn_list ::= windowdefn_list COMMA windowdefn */
 {
-  assert( yymsp[0].minor.yy295!=0 );
-  sqlite3WindowChain(pParse, yymsp[0].minor.yy295, yymsp[-2].minor.yy295);
-  yymsp[0].minor.yy295->pNextWin = yymsp[-2].minor.yy295;
-  yylhsminor.yy295 = yymsp[0].minor.yy295;
+  assert( yymsp[0].minor.yy303!=0 );
+  sqlite3WindowChain(pParse, yymsp[0].minor.yy303, yymsp[-2].minor.yy303);
+  yymsp[0].minor.yy303->pNextWin = yymsp[-2].minor.yy303;
+  yylhsminor.yy303 = yymsp[0].minor.yy303;
 }
-  yymsp[-2].minor.yy295 = yylhsminor.yy295;
+  yymsp[-2].minor.yy303 = yylhsminor.yy303;
         break;
-      case 290: /* windowdefn ::= nm AS LP window RP */
+      case 296: /* windowdefn ::= nm AS LP window RP */
 {
-  if( ALWAYS(yymsp[-1].minor.yy295) ){
-    yymsp[-1].minor.yy295->zName = sqlite3DbStrNDup(pParse->db, yymsp[-4].minor.yy0.z, yymsp[-4].minor.yy0.n);
+  if( ALWAYS(yymsp[-1].minor.yy303) ){
+    yymsp[-1].minor.yy303->zName = sqlite3DbStrNDup(pParse->db, yymsp[-4].minor.yy0.z, yymsp[-4].minor.yy0.n);
   }
-  yylhsminor.yy295 = yymsp[-1].minor.yy295;
+  yylhsminor.yy303 = yymsp[-1].minor.yy303;
 }
-  yymsp[-4].minor.yy295 = yylhsminor.yy295;
+  yymsp[-4].minor.yy303 = yylhsminor.yy303;
         break;
-      case 291: /* window ::= PARTITION BY nexprlist orderby_opt frame_opt */
+      case 297: /* window ::= PARTITION BY nexprlist orderby_opt frame_opt */
 {
-  yymsp[-4].minor.yy295 = sqlite3WindowAssemble(pParse, yymsp[0].minor.yy295, yymsp[-2].minor.yy434, yymsp[-1].minor.yy434, 0);
+  yymsp[-4].minor.yy303 = sqlite3WindowAssemble(pParse, yymsp[0].minor.yy303, yymsp[-2].minor.yy242, yymsp[-1].minor.yy242, 0);
 }
         break;
-      case 292: /* window ::= nm PARTITION BY nexprlist orderby_opt frame_opt */
+      case 298: /* window ::= nm PARTITION BY nexprlist orderby_opt frame_opt */
 {
-  yylhsminor.yy295 = sqlite3WindowAssemble(pParse, yymsp[0].minor.yy295, yymsp[-2].minor.yy434, yymsp[-1].minor.yy434, &yymsp[-5].minor.yy0);
+  yylhsminor.yy303 = sqlite3WindowAssemble(pParse, yymsp[0].minor.yy303, yymsp[-2].minor.yy242, yymsp[-1].minor.yy242, &yymsp[-5].minor.yy0);
 }
-  yymsp[-5].minor.yy295 = yylhsminor.yy295;
+  yymsp[-5].minor.yy303 = yylhsminor.yy303;
         break;
-      case 293: /* window ::= ORDER BY sortlist frame_opt */
+      case 299: /* window ::= ORDER BY sortlist frame_opt */
 {
-  yymsp[-3].minor.yy295 = sqlite3WindowAssemble(pParse, yymsp[0].minor.yy295, 0, yymsp[-1].minor.yy434, 0);
+  yymsp[-3].minor.yy303 = sqlite3WindowAssemble(pParse, yymsp[0].minor.yy303, 0, yymsp[-1].minor.yy242, 0);
 }
         break;
-      case 294: /* window ::= nm ORDER BY sortlist frame_opt */
+      case 300: /* window ::= nm ORDER BY sortlist frame_opt */
 {
-  yylhsminor.yy295 = sqlite3WindowAssemble(pParse, yymsp[0].minor.yy295, 0, yymsp[-1].minor.yy434, &yymsp[-4].minor.yy0);
+  yylhsminor.yy303 = sqlite3WindowAssemble(pParse, yymsp[0].minor.yy303, 0, yymsp[-1].minor.yy242, &yymsp[-4].minor.yy0);
 }
-  yymsp[-4].minor.yy295 = yylhsminor.yy295;
+  yymsp[-4].minor.yy303 = yylhsminor.yy303;
         break;
-      case 295: /* window ::= frame_opt */
+      case 301: /* window ::= frame_opt */
+      case 320: /* filter_over ::= over_clause */ yytestcase(yyruleno==320);
 {
-  yylhsminor.yy295 = yymsp[0].minor.yy295;
+  yylhsminor.yy303 = yymsp[0].minor.yy303;
 }
-  yymsp[0].minor.yy295 = yylhsminor.yy295;
+  yymsp[0].minor.yy303 = yylhsminor.yy303;
         break;
-      case 296: /* window ::= nm frame_opt */
+      case 302: /* window ::= nm frame_opt */
 {
-  yylhsminor.yy295 = sqlite3WindowAssemble(pParse, yymsp[0].minor.yy295, 0, 0, &yymsp[-1].minor.yy0);
+  yylhsminor.yy303 = sqlite3WindowAssemble(pParse, yymsp[0].minor.yy303, 0, 0, &yymsp[-1].minor.yy0);
 }
-  yymsp[-1].minor.yy295 = yylhsminor.yy295;
+  yymsp[-1].minor.yy303 = yylhsminor.yy303;
         break;
-      case 297: /* frame_opt ::= */
+      case 303: /* frame_opt ::= */
 { 
-  yymsp[1].minor.yy295 = sqlite3WindowAlloc(pParse, 0, TK_UNBOUNDED, 0, TK_CURRENT, 0, 0);
+  yymsp[1].minor.yy303 = sqlite3WindowAlloc(pParse, 0, TK_UNBOUNDED, 0, TK_CURRENT, 0, 0);
 }
         break;
-      case 298: /* frame_opt ::= range_or_rows frame_bound_s frame_exclude_opt */
+      case 304: /* frame_opt ::= range_or_rows frame_bound_s frame_exclude_opt */
 { 
-  yylhsminor.yy295 = sqlite3WindowAlloc(pParse, yymsp[-2].minor.yy494, yymsp[-1].minor.yy201.eType, yymsp[-1].minor.yy201.pExpr, TK_CURRENT, 0, yymsp[0].minor.yy238);
+  yylhsminor.yy303 = sqlite3WindowAlloc(pParse, yymsp[-2].minor.yy192, yymsp[-1].minor.yy77.eType, yymsp[-1].minor.yy77.pExpr, TK_CURRENT, 0, yymsp[0].minor.yy58);
 }
-  yymsp[-2].minor.yy295 = yylhsminor.yy295;
+  yymsp[-2].minor.yy303 = yylhsminor.yy303;
         break;
-      case 299: /* frame_opt ::= range_or_rows BETWEEN frame_bound_s AND frame_bound_e frame_exclude_opt */
+      case 305: /* frame_opt ::= range_or_rows BETWEEN frame_bound_s AND frame_bound_e frame_exclude_opt */
 { 
-  yylhsminor.yy295 = sqlite3WindowAlloc(pParse, yymsp[-5].minor.yy494, yymsp[-3].minor.yy201.eType, yymsp[-3].minor.yy201.pExpr, yymsp[-1].minor.yy201.eType, yymsp[-1].minor.yy201.pExpr, yymsp[0].minor.yy238);
+  yylhsminor.yy303 = sqlite3WindowAlloc(pParse, yymsp[-5].minor.yy192, yymsp[-3].minor.yy77.eType, yymsp[-3].minor.yy77.pExpr, yymsp[-1].minor.yy77.eType, yymsp[-1].minor.yy77.pExpr, yymsp[0].minor.yy58);
 }
-  yymsp[-5].minor.yy295 = yylhsminor.yy295;
+  yymsp[-5].minor.yy303 = yylhsminor.yy303;
         break;
-      case 301: /* frame_bound_s ::= frame_bound */
-      case 303: /* frame_bound_e ::= frame_bound */ yytestcase(yyruleno==303);
-{yylhsminor.yy201 = yymsp[0].minor.yy201;}
-  yymsp[0].minor.yy201 = yylhsminor.yy201;
+      case 307: /* frame_bound_s ::= frame_bound */
+      case 309: /* frame_bound_e ::= frame_bound */ yytestcase(yyruleno==309);
+{yylhsminor.yy77 = yymsp[0].minor.yy77;}
+  yymsp[0].minor.yy77 = yylhsminor.yy77;
         break;
-      case 302: /* frame_bound_s ::= UNBOUNDED PRECEDING */
-      case 304: /* frame_bound_e ::= UNBOUNDED FOLLOWING */ yytestcase(yyruleno==304);
-      case 306: /* frame_bound ::= CURRENT ROW */ yytestcase(yyruleno==306);
-{yylhsminor.yy201.eType = yymsp[-1].major; yylhsminor.yy201.pExpr = 0;}
-  yymsp[-1].minor.yy201 = yylhsminor.yy201;
+      case 308: /* frame_bound_s ::= UNBOUNDED PRECEDING */
+      case 310: /* frame_bound_e ::= UNBOUNDED FOLLOWING */ yytestcase(yyruleno==310);
+      case 312: /* frame_bound ::= CURRENT ROW */ yytestcase(yyruleno==312);
+{yylhsminor.yy77.eType = yymsp[-1].major; yylhsminor.yy77.pExpr = 0;}
+  yymsp[-1].minor.yy77 = yylhsminor.yy77;
         break;
-      case 305: /* frame_bound ::= expr PRECEDING|FOLLOWING */
-{yylhsminor.yy201.eType = yymsp[0].major; yylhsminor.yy201.pExpr = yymsp[-1].minor.yy524;}
-  yymsp[-1].minor.yy201 = yylhsminor.yy201;
+      case 311: /* frame_bound ::= expr PRECEDING|FOLLOWING */
+{yylhsminor.yy77.eType = yymsp[0].major; yylhsminor.yy77.pExpr = yymsp[-1].minor.yy202;}
+  yymsp[-1].minor.yy77 = yylhsminor.yy77;
         break;
-      case 307: /* frame_exclude_opt ::= */
-{yymsp[1].minor.yy238 = 0;}
+      case 313: /* frame_exclude_opt ::= */
+{yymsp[1].minor.yy58 = 0;}
         break;
-      case 308: /* frame_exclude_opt ::= EXCLUDE frame_exclude */
-{yymsp[-1].minor.yy238 = yymsp[0].minor.yy238;}
+      case 314: /* frame_exclude_opt ::= EXCLUDE frame_exclude */
+{yymsp[-1].minor.yy58 = yymsp[0].minor.yy58;}
         break;
-      case 309: /* frame_exclude ::= NO OTHERS */
-      case 310: /* frame_exclude ::= CURRENT ROW */ yytestcase(yyruleno==310);
-{yymsp[-1].minor.yy238 = yymsp[-1].major; /*A-overwrites-X*/}
+      case 315: /* frame_exclude ::= NO OTHERS */
+      case 316: /* frame_exclude ::= CURRENT ROW */ yytestcase(yyruleno==316);
+{yymsp[-1].minor.yy58 = yymsp[-1].major; /*A-overwrites-X*/}
         break;
-      case 311: /* frame_exclude ::= GROUP|TIES */
-{yymsp[0].minor.yy238 = yymsp[0].major; /*A-overwrites-X*/}
+      case 317: /* frame_exclude ::= GROUP|TIES */
+{yymsp[0].minor.yy58 = yymsp[0].major; /*A-overwrites-X*/}
         break;
-      case 312: /* window_clause ::= WINDOW windowdefn_list */
-{ yymsp[-1].minor.yy295 = yymsp[0].minor.yy295; }
+      case 318: /* window_clause ::= WINDOW windowdefn_list */
+{ yymsp[-1].minor.yy303 = yymsp[0].minor.yy303; }
         break;
-      case 313: /* over_clause ::= filter_opt OVER LP window RP */
+      case 319: /* filter_over ::= filter_clause over_clause */
 {
-  yylhsminor.yy295 = yymsp[-1].minor.yy295;
-  assert( yylhsminor.yy295!=0 );
-  yylhsminor.yy295->pFilter = yymsp[-4].minor.yy524;
+  yymsp[0].minor.yy303->pFilter = yymsp[-1].minor.yy202;
+  yylhsminor.yy303 = yymsp[0].minor.yy303;
 }
-  yymsp[-4].minor.yy295 = yylhsminor.yy295;
+  yymsp[-1].minor.yy303 = yylhsminor.yy303;
         break;
-      case 314: /* over_clause ::= filter_opt OVER nm */
+      case 321: /* filter_over ::= filter_clause */
 {
-  yylhsminor.yy295 = (Window*)sqlite3DbMallocZero(pParse->db, sizeof(Window));
-  if( yylhsminor.yy295 ){
-    yylhsminor.yy295->zName = sqlite3DbStrNDup(pParse->db, yymsp[0].minor.yy0.z, yymsp[0].minor.yy0.n);
-    yylhsminor.yy295->pFilter = yymsp[-2].minor.yy524;
+  yylhsminor.yy303 = (Window*)sqlite3DbMallocZero(pParse->db, sizeof(Window));
+  if( yylhsminor.yy303 ){
+    yylhsminor.yy303->eFrmType = TK_FILTER;
+    yylhsminor.yy303->pFilter = yymsp[0].minor.yy202;
   }else{
-    sqlite3ExprDelete(pParse->db, yymsp[-2].minor.yy524);
+    sqlite3ExprDelete(pParse->db, yymsp[0].minor.yy202);
+  }
+}
+  yymsp[0].minor.yy303 = yylhsminor.yy303;
+        break;
+      case 322: /* over_clause ::= OVER LP window RP */
+{
+  yymsp[-3].minor.yy303 = yymsp[-1].minor.yy303;
+  assert( yymsp[-3].minor.yy303!=0 );
+}
+        break;
+      case 323: /* over_clause ::= OVER nm */
+{
+  yymsp[-1].minor.yy303 = (Window*)sqlite3DbMallocZero(pParse->db, sizeof(Window));
+  if( yymsp[-1].minor.yy303 ){
+    yymsp[-1].minor.yy303->zName = sqlite3DbStrNDup(pParse->db, yymsp[0].minor.yy0.z, yymsp[0].minor.yy0.n);
   }
 }
-  yymsp[-2].minor.yy295 = yylhsminor.yy295;
         break;
-      case 316: /* filter_opt ::= FILTER LP WHERE expr RP */
-{ yymsp[-4].minor.yy524 = yymsp[-1].minor.yy524; }
+      case 324: /* filter_clause ::= FILTER LP WHERE expr RP */
+{ yymsp[-4].minor.yy202 = yymsp[-1].minor.yy202; }
         break;
       default:
-      /* (317) input ::= cmdlist */ yytestcase(yyruleno==317);
-      /* (318) cmdlist ::= cmdlist ecmd */ yytestcase(yyruleno==318);
-      /* (319) cmdlist ::= ecmd (OPTIMIZED OUT) */ assert(yyruleno!=319);
-      /* (320) ecmd ::= SEMI */ yytestcase(yyruleno==320);
-      /* (321) ecmd ::= cmdx SEMI */ yytestcase(yyruleno==321);
-      /* (322) ecmd ::= explain cmdx */ yytestcase(yyruleno==322);
-      /* (323) trans_opt ::= */ yytestcase(yyruleno==323);
-      /* (324) trans_opt ::= TRANSACTION */ yytestcase(yyruleno==324);
-      /* (325) trans_opt ::= TRANSACTION nm */ yytestcase(yyruleno==325);
-      /* (326) savepoint_opt ::= SAVEPOINT */ yytestcase(yyruleno==326);
-      /* (327) savepoint_opt ::= */ yytestcase(yyruleno==327);
-      /* (328) cmd ::= create_table create_table_args */ yytestcase(yyruleno==328);
-      /* (329) columnlist ::= columnlist COMMA columnname carglist */ yytestcase(yyruleno==329);
-      /* (330) columnlist ::= columnname carglist */ yytestcase(yyruleno==330);
-      /* (331) nm ::= ID|INDEXED */ yytestcase(yyruleno==331);
-      /* (332) nm ::= STRING */ yytestcase(yyruleno==332);
-      /* (333) nm ::= JOIN_KW */ yytestcase(yyruleno==333);
-      /* (334) typetoken ::= typename */ yytestcase(yyruleno==334);
-      /* (335) typename ::= ID|STRING */ yytestcase(yyruleno==335);
-      /* (336) signed ::= plus_num (OPTIMIZED OUT) */ assert(yyruleno!=336);
-      /* (337) signed ::= minus_num (OPTIMIZED OUT) */ assert(yyruleno!=337);
-      /* (338) carglist ::= carglist ccons */ yytestcase(yyruleno==338);
-      /* (339) carglist ::= */ yytestcase(yyruleno==339);
-      /* (340) ccons ::= NULL onconf */ yytestcase(yyruleno==340);
-      /* (341) conslist_opt ::= COMMA conslist */ yytestcase(yyruleno==341);
-      /* (342) conslist ::= conslist tconscomma tcons */ yytestcase(yyruleno==342);
-      /* (343) conslist ::= tcons (OPTIMIZED OUT) */ assert(yyruleno!=343);
-      /* (344) tconscomma ::= */ yytestcase(yyruleno==344);
-      /* (345) defer_subclause_opt ::= defer_subclause (OPTIMIZED OUT) */ assert(yyruleno!=345);
-      /* (346) resolvetype ::= raisetype (OPTIMIZED OUT) */ assert(yyruleno!=346);
-      /* (347) selectnowith ::= oneselect (OPTIMIZED OUT) */ assert(yyruleno!=347);
-      /* (348) oneselect ::= values */ yytestcase(yyruleno==348);
-      /* (349) sclp ::= selcollist COMMA */ yytestcase(yyruleno==349);
-      /* (350) as ::= ID|STRING */ yytestcase(yyruleno==350);
-      /* (351) expr ::= term (OPTIMIZED OUT) */ assert(yyruleno!=351);
-      /* (352) likeop ::= LIKE_KW|MATCH */ yytestcase(yyruleno==352);
-      /* (353) exprlist ::= nexprlist */ yytestcase(yyruleno==353);
-      /* (354) nmnum ::= plus_num (OPTIMIZED OUT) */ assert(yyruleno!=354);
-      /* (355) nmnum ::= nm (OPTIMIZED OUT) */ assert(yyruleno!=355);
-      /* (356) nmnum ::= ON */ yytestcase(yyruleno==356);
-      /* (357) nmnum ::= DELETE */ yytestcase(yyruleno==357);
-      /* (358) nmnum ::= DEFAULT */ yytestcase(yyruleno==358);
-      /* (359) plus_num ::= INTEGER|FLOAT */ yytestcase(yyruleno==359);
-      /* (360) foreach_clause ::= */ yytestcase(yyruleno==360);
-      /* (361) foreach_clause ::= FOR EACH ROW */ yytestcase(yyruleno==361);
-      /* (362) trnm ::= nm */ yytestcase(yyruleno==362);
-      /* (363) tridxby ::= */ yytestcase(yyruleno==363);
-      /* (364) database_kw_opt ::= DATABASE */ yytestcase(yyruleno==364);
-      /* (365) database_kw_opt ::= */ yytestcase(yyruleno==365);
-      /* (366) kwcolumn_opt ::= */ yytestcase(yyruleno==366);
-      /* (367) kwcolumn_opt ::= COLUMNKW */ yytestcase(yyruleno==367);
-      /* (368) vtabarglist ::= vtabarg */ yytestcase(yyruleno==368);
-      /* (369) vtabarglist ::= vtabarglist COMMA vtabarg */ yytestcase(yyruleno==369);
-      /* (370) vtabarg ::= vtabarg vtabargtoken */ yytestcase(yyruleno==370);
-      /* (371) anylist ::= */ yytestcase(yyruleno==371);
-      /* (372) anylist ::= anylist LP anylist RP */ yytestcase(yyruleno==372);
-      /* (373) anylist ::= anylist ANY */ yytestcase(yyruleno==373);
-      /* (374) with ::= */ yytestcase(yyruleno==374);
+      /* (325) input ::= cmdlist */ yytestcase(yyruleno==325);
+      /* (326) cmdlist ::= cmdlist ecmd */ yytestcase(yyruleno==326);
+      /* (327) cmdlist ::= ecmd (OPTIMIZED OUT) */ assert(yyruleno!=327);
+      /* (328) ecmd ::= SEMI */ yytestcase(yyruleno==328);
+      /* (329) ecmd ::= cmdx SEMI */ yytestcase(yyruleno==329);
+      /* (330) ecmd ::= explain cmdx SEMI (NEVER REDUCES) */ assert(yyruleno!=330);
+      /* (331) trans_opt ::= */ yytestcase(yyruleno==331);
+      /* (332) trans_opt ::= TRANSACTION */ yytestcase(yyruleno==332);
+      /* (333) trans_opt ::= TRANSACTION nm */ yytestcase(yyruleno==333);
+      /* (334) savepoint_opt ::= SAVEPOINT */ yytestcase(yyruleno==334);
+      /* (335) savepoint_opt ::= */ yytestcase(yyruleno==335);
+      /* (336) cmd ::= create_table create_table_args */ yytestcase(yyruleno==336);
+      /* (337) columnlist ::= columnlist COMMA columnname carglist */ yytestcase(yyruleno==337);
+      /* (338) columnlist ::= columnname carglist */ yytestcase(yyruleno==338);
+      /* (339) nm ::= ID|INDEXED */ yytestcase(yyruleno==339);
+      /* (340) nm ::= STRING */ yytestcase(yyruleno==340);
+      /* (341) nm ::= JOIN_KW */ yytestcase(yyruleno==341);
+      /* (342) typetoken ::= typename */ yytestcase(yyruleno==342);
+      /* (343) typename ::= ID|STRING */ yytestcase(yyruleno==343);
+      /* (344) signed ::= plus_num (OPTIMIZED OUT) */ assert(yyruleno!=344);
+      /* (345) signed ::= minus_num (OPTIMIZED OUT) */ assert(yyruleno!=345);
+      /* (346) carglist ::= carglist ccons */ yytestcase(yyruleno==346);
+      /* (347) carglist ::= */ yytestcase(yyruleno==347);
+      /* (348) ccons ::= NULL onconf */ yytestcase(yyruleno==348);
+      /* (349) ccons ::= GENERATED ALWAYS AS generated */ yytestcase(yyruleno==349);
+      /* (350) ccons ::= AS generated */ yytestcase(yyruleno==350);
+      /* (351) conslist_opt ::= COMMA conslist */ yytestcase(yyruleno==351);
+      /* (352) conslist ::= conslist tconscomma tcons */ yytestcase(yyruleno==352);
+      /* (353) conslist ::= tcons (OPTIMIZED OUT) */ assert(yyruleno!=353);
+      /* (354) tconscomma ::= */ yytestcase(yyruleno==354);
+      /* (355) defer_subclause_opt ::= defer_subclause (OPTIMIZED OUT) */ assert(yyruleno!=355);
+      /* (356) resolvetype ::= raisetype (OPTIMIZED OUT) */ assert(yyruleno!=356);
+      /* (357) selectnowith ::= oneselect (OPTIMIZED OUT) */ assert(yyruleno!=357);
+      /* (358) oneselect ::= values */ yytestcase(yyruleno==358);
+      /* (359) sclp ::= selcollist COMMA */ yytestcase(yyruleno==359);
+      /* (360) as ::= ID|STRING */ yytestcase(yyruleno==360);
+      /* (361) expr ::= term (OPTIMIZED OUT) */ assert(yyruleno!=361);
+      /* (362) likeop ::= LIKE_KW|MATCH */ yytestcase(yyruleno==362);
+      /* (363) exprlist ::= nexprlist */ yytestcase(yyruleno==363);
+      /* (364) nmnum ::= plus_num (OPTIMIZED OUT) */ assert(yyruleno!=364);
+      /* (365) nmnum ::= nm (OPTIMIZED OUT) */ assert(yyruleno!=365);
+      /* (366) nmnum ::= ON */ yytestcase(yyruleno==366);
+      /* (367) nmnum ::= DELETE */ yytestcase(yyruleno==367);
+      /* (368) nmnum ::= DEFAULT */ yytestcase(yyruleno==368);
+      /* (369) plus_num ::= INTEGER|FLOAT */ yytestcase(yyruleno==369);
+      /* (370) foreach_clause ::= */ yytestcase(yyruleno==370);
+      /* (371) foreach_clause ::= FOR EACH ROW */ yytestcase(yyruleno==371);
+      /* (372) trnm ::= nm */ yytestcase(yyruleno==372);
+      /* (373) tridxby ::= */ yytestcase(yyruleno==373);
+      /* (374) database_kw_opt ::= DATABASE */ yytestcase(yyruleno==374);
+      /* (375) database_kw_opt ::= */ yytestcase(yyruleno==375);
+      /* (376) kwcolumn_opt ::= */ yytestcase(yyruleno==376);
+      /* (377) kwcolumn_opt ::= COLUMNKW */ yytestcase(yyruleno==377);
+      /* (378) vtabarglist ::= vtabarg */ yytestcase(yyruleno==378);
+      /* (379) vtabarglist ::= vtabarglist COMMA vtabarg */ yytestcase(yyruleno==379);
+      /* (380) vtabarg ::= vtabarg vtabargtoken */ yytestcase(yyruleno==380);
+      /* (381) anylist ::= */ yytestcase(yyruleno==381);
+      /* (382) anylist ::= anylist LP anylist RP */ yytestcase(yyruleno==382);
+      /* (383) anylist ::= anylist ANY */ yytestcase(yyruleno==383);
+      /* (384) with ::= */ yytestcase(yyruleno==384);
         break;
 /********** End reduce actions ************************************************/
   };
@@ -153010,13 +157314,12 @@ SQLITE_PRIVATE void sqlite3Parser(
 */
 SQLITE_PRIVATE int sqlite3ParserFallback(int iToken){
 #ifdef YYFALLBACK
-  if( iToken<(int)(sizeof(yyFallback)/sizeof(yyFallback[0])) ){
-    return yyFallback[iToken];
-  }
+  assert( iToken<(int)(sizeof(yyFallback)/sizeof(yyFallback[0])) );
+  return yyFallback[iToken];
 #else
   (void)iToken;
-#endif
   return 0;
+#endif
 }
 
 /************** End of parse.c ***********************************************/
@@ -153181,145 +157484,279 @@ const unsigned char ebcdicToAscii[] = {
 ** is substantially reduced.  This is important for embedded applications
 ** on platforms with limited memory.
 */
-/* Hash score: 214 */
-/* zKWText[] encodes 950 bytes of keyword text in 629 bytes */
+/* Hash score: 227 */
+/* zKWText[] encodes 984 bytes of keyword text in 648 bytes */
 /*   REINDEXEDESCAPEACHECKEYBEFOREIGNOREGEXPLAINSTEADDATABASELECT       */
-/*   ABLEFTHENDEFERRABLELSEXCLUDELETEMPORARYCONSTRAINTERSECTIES         */
-/*   AVEPOINTOFFSETRANSACTIONATURALTERAISEXCEPTRIGGEREFERENCES          */
-/*   UNIQUERYWITHOUTERELEASEXCLUSIVEXISTSATTACHAVINGLOBEGINNERANGE      */
-/*   BETWEENOTHINGROUPSCASCADETACHCASECOLLATECREATECURRENT_DATE         */
-/*   IMMEDIATEJOINSERTLIKEMATCHPLANALYZEPRAGMABORTUPDATEVALUES          */
-/*   VIRTUALIMITWHENOTNULLWHERECURSIVEAFTERENAMEANDEFAULT               */
-/*   AUTOINCREMENTCASTCOLUMNCOMMITCONFLICTCROSSCURRENT_TIMESTAMP        */
-/*   ARTITIONDEFERREDISTINCTDROPRECEDINGFAILFILTEREPLACEFOLLOWING       */
-/*   FROMFULLIFISNULLORDERESTRICTOTHERSOVERIGHTROLLBACKROWS             */
+/*   ABLEFTHENDEFERRABLELSEXCLUDELETEMPORARYISNULLSAVEPOINTERSECT       */
+/*   IESNOTNULLIKEXCEPTRANSACTIONATURALTERAISEXCLUSIVEXISTS             */
+/*   CONSTRAINTOFFSETRIGGERANGENERATEDETACHAVINGLOBEGINNEREFERENCES     */
+/*   UNIQUERYWITHOUTERELEASEATTACHBETWEENOTHINGROUPSCASCADEFAULT        */
+/*   CASECOLLATECREATECURRENT_DATEIMMEDIATEJOINSERTMATCHPLANALYZE       */
+/*   PRAGMABORTUPDATEVALUESVIRTUALWAYSWHENWHERECURSIVEAFTERENAMEAND     */
+/*   EFERREDISTINCTAUTOINCREMENTCASTCOLUMNCOMMITCONFLICTCROSS           */
+/*   CURRENT_TIMESTAMPARTITIONDROPRECEDINGFAILASTFILTEREPLACEFIRST      */
+/*   FOLLOWINGFROMFULLIMITIFORDERESTRICTOTHERSOVERIGHTROLLBACKROWS      */
 /*   UNBOUNDEDUNIONUSINGVACUUMVIEWINDOWBYINITIALLYPRIMARY               */
-static const char zKWText[628] = {
+static const char zKWText[647] = {
   'R','E','I','N','D','E','X','E','D','E','S','C','A','P','E','A','C','H',
   'E','C','K','E','Y','B','E','F','O','R','E','I','G','N','O','R','E','G',
   'E','X','P','L','A','I','N','S','T','E','A','D','D','A','T','A','B','A',
   'S','E','L','E','C','T','A','B','L','E','F','T','H','E','N','D','E','F',
   'E','R','R','A','B','L','E','L','S','E','X','C','L','U','D','E','L','E',
-  'T','E','M','P','O','R','A','R','Y','C','O','N','S','T','R','A','I','N',
-  'T','E','R','S','E','C','T','I','E','S','A','V','E','P','O','I','N','T',
-  'O','F','F','S','E','T','R','A','N','S','A','C','T','I','O','N','A','T',
-  'U','R','A','L','T','E','R','A','I','S','E','X','C','E','P','T','R','I',
-  'G','G','E','R','E','F','E','R','E','N','C','E','S','U','N','I','Q','U',
-  'E','R','Y','W','I','T','H','O','U','T','E','R','E','L','E','A','S','E',
-  'X','C','L','U','S','I','V','E','X','I','S','T','S','A','T','T','A','C',
-  'H','A','V','I','N','G','L','O','B','E','G','I','N','N','E','R','A','N',
-  'G','E','B','E','T','W','E','E','N','O','T','H','I','N','G','R','O','U',
-  'P','S','C','A','S','C','A','D','E','T','A','C','H','C','A','S','E','C',
-  'O','L','L','A','T','E','C','R','E','A','T','E','C','U','R','R','E','N',
-  'T','_','D','A','T','E','I','M','M','E','D','I','A','T','E','J','O','I',
-  'N','S','E','R','T','L','I','K','E','M','A','T','C','H','P','L','A','N',
-  'A','L','Y','Z','E','P','R','A','G','M','A','B','O','R','T','U','P','D',
-  'A','T','E','V','A','L','U','E','S','V','I','R','T','U','A','L','I','M',
-  'I','T','W','H','E','N','O','T','N','U','L','L','W','H','E','R','E','C',
-  'U','R','S','I','V','E','A','F','T','E','R','E','N','A','M','E','A','N',
-  'D','E','F','A','U','L','T','A','U','T','O','I','N','C','R','E','M','E',
-  'N','T','C','A','S','T','C','O','L','U','M','N','C','O','M','M','I','T',
-  'C','O','N','F','L','I','C','T','C','R','O','S','S','C','U','R','R','E',
-  'N','T','_','T','I','M','E','S','T','A','M','P','A','R','T','I','T','I',
-  'O','N','D','E','F','E','R','R','E','D','I','S','T','I','N','C','T','D',
-  'R','O','P','R','E','C','E','D','I','N','G','F','A','I','L','F','I','L',
-  'T','E','R','E','P','L','A','C','E','F','O','L','L','O','W','I','N','G',
-  'F','R','O','M','F','U','L','L','I','F','I','S','N','U','L','L','O','R',
-  'D','E','R','E','S','T','R','I','C','T','O','T','H','E','R','S','O','V',
-  'E','R','I','G','H','T','R','O','L','L','B','A','C','K','R','O','W','S',
-  'U','N','B','O','U','N','D','E','D','U','N','I','O','N','U','S','I','N',
-  'G','V','A','C','U','U','M','V','I','E','W','I','N','D','O','W','B','Y',
-  'I','N','I','T','I','A','L','L','Y','P','R','I','M','A','R','Y',
+  'T','E','M','P','O','R','A','R','Y','I','S','N','U','L','L','S','A','V',
+  'E','P','O','I','N','T','E','R','S','E','C','T','I','E','S','N','O','T',
+  'N','U','L','L','I','K','E','X','C','E','P','T','R','A','N','S','A','C',
+  'T','I','O','N','A','T','U','R','A','L','T','E','R','A','I','S','E','X',
+  'C','L','U','S','I','V','E','X','I','S','T','S','C','O','N','S','T','R',
+  'A','I','N','T','O','F','F','S','E','T','R','I','G','G','E','R','A','N',
+  'G','E','N','E','R','A','T','E','D','E','T','A','C','H','A','V','I','N',
+  'G','L','O','B','E','G','I','N','N','E','R','E','F','E','R','E','N','C',
+  'E','S','U','N','I','Q','U','E','R','Y','W','I','T','H','O','U','T','E',
+  'R','E','L','E','A','S','E','A','T','T','A','C','H','B','E','T','W','E',
+  'E','N','O','T','H','I','N','G','R','O','U','P','S','C','A','S','C','A',
+  'D','E','F','A','U','L','T','C','A','S','E','C','O','L','L','A','T','E',
+  'C','R','E','A','T','E','C','U','R','R','E','N','T','_','D','A','T','E',
+  'I','M','M','E','D','I','A','T','E','J','O','I','N','S','E','R','T','M',
+  'A','T','C','H','P','L','A','N','A','L','Y','Z','E','P','R','A','G','M',
+  'A','B','O','R','T','U','P','D','A','T','E','V','A','L','U','E','S','V',
+  'I','R','T','U','A','L','W','A','Y','S','W','H','E','N','W','H','E','R',
+  'E','C','U','R','S','I','V','E','A','F','T','E','R','E','N','A','M','E',
+  'A','N','D','E','F','E','R','R','E','D','I','S','T','I','N','C','T','A',
+  'U','T','O','I','N','C','R','E','M','E','N','T','C','A','S','T','C','O',
+  'L','U','M','N','C','O','M','M','I','T','C','O','N','F','L','I','C','T',
+  'C','R','O','S','S','C','U','R','R','E','N','T','_','T','I','M','E','S',
+  'T','A','M','P','A','R','T','I','T','I','O','N','D','R','O','P','R','E',
+  'C','E','D','I','N','G','F','A','I','L','A','S','T','F','I','L','T','E',
+  'R','E','P','L','A','C','E','F','I','R','S','T','F','O','L','L','O','W',
+  'I','N','G','F','R','O','M','F','U','L','L','I','M','I','T','I','F','O',
+  'R','D','E','R','E','S','T','R','I','C','T','O','T','H','E','R','S','O',
+  'V','E','R','I','G','H','T','R','O','L','L','B','A','C','K','R','O','W',
+  'S','U','N','B','O','U','N','D','E','D','U','N','I','O','N','U','S','I',
+  'N','G','V','A','C','U','U','M','V','I','E','W','I','N','D','O','W','B',
+  'Y','I','N','I','T','I','A','L','L','Y','P','R','I','M','A','R','Y',
 };
 /* aKWHash[i] is the hash value for the i-th keyword */
 static const unsigned char aKWHash[127] = {
-    75, 111, 127,  73, 108,  29,   0,   0,  83,   0,  77,  63,   0,
-    37,  33,  78,  15,   0, 126,  86,  57, 120, 128,  19,   0,   0,
-   133,   0, 131, 123,   0,  22,  98,   0,   9,   0,   0, 117,  71,
-     0,  69,   6,   0,  49,  95, 140,   0, 129, 106,   0,   0,  54,
-     0, 109,  24,   0,  17,   0, 134,  56,  23,  26,   5,  58, 135,
-   101,   0,   0, 139, 112,  62, 138,  59, 115,  65,   0,  96,   0,
-   105,  45,   0, 104,   0,   0,   0, 100,  97, 102, 107, 119,  14,
-    31, 118,   0,  81,   0, 136, 116, 137,  61, 124, 132,  80, 121,
-    88,  30,  85,   0,   0,  99,  35, 125, 122,   0, 130,   0,   0,
-    41,   0,  91,  89,  90,   0,  20,  87, 113,  82,
+    84, 102, 132,  82, 114,  29,   0,   0,  91,   0,  85,  72,   0,
+    53,  35,  86,  15,   0,  42,  94,  54, 126, 133,  19,   0,   0,
+   138,   0,  40, 128,   0,  22, 104,   0,   9,   0,   0, 122,  80,
+     0,  78,   6,   0,  65,  99, 145,   0, 134, 112,   0,   0,  48,
+     0, 100,  24,   0,  17,   0,  27,  70,  23,  26,   5,  60, 140,
+   107, 121,   0,  73, 101,  71, 143,  61, 119,  74,   0,  49,   0,
+    11,  41,   0, 110,   0,   0,   0, 106,  10, 108, 113, 124,  14,
+    50, 123,   0,  89,   0,  18, 120, 142,  56, 129, 137,  88,  83,
+    37,  30, 125,   0,   0, 105,  51, 130, 127,   0,  34,   0,   0,
+    44,   0,  95,  38,  39,   0,  20,  45, 116,  90,
 };
 /* aKWNext[] forms the hash collision chain.  If aKWHash[i]==0
 ** then the i-th keyword has no more hash collisions.  Otherwise,
 ** the next keyword with the same hash is aKWHash[i]-1. */
-static const unsigned char aKWNext[140] = {
-     0,   0,   0,   0,   4,   0,   0,   0,   0,   0,   0,   0,   0,
-     0,   2,   0,   0,   0,   0,   0,   0,  13,   0,   0,   0,   0,
-     0,   0,   0,  21,   0,   0,  12,   0,   0,   0,   0,   0,   0,
-     7,   0,   0,   0,   0,   0,   0,   0,   0,   0,   0,   0,   0,
-    51,  28,   0,   0,  38,   0,   0,   0,  44,   0,   0,   0,   3,
-     0,   0,  67,   1,  66,   0,   0,   0,  36,   0,  47,   0,   0,
-     0,   0,   0,  48,  50,  76,   0,   0,  42,   0,  60,   0,   0,
-     0,  43,   0,  16,  55,  10,   0,   0,   0,   0,   0,   0,   0,
-    11,  72,  93,   0,   0,   8,   0, 110,   0, 103,  40,  53,  70,
-     0, 114,   0,  74,  52,   0,   0,  92,  39,  46,   0,  68,  32,
-    84,   0,  34,  27,  25,  18,  94,   0,  64,  79,
+static const unsigned char aKWNext[145] = {
+     0,   0,   0,   0,   4,   0,  43,   0,   0, 103, 111,   0,   0,
+     0,   2,   0,   0, 141,   0,   0,   0,  13,   0,   0,   0,   0,
+   139,   0,   0, 118,  52,   0,   0, 135,  12,   0,   0,  62,   0,
+   136,   0, 131,   0,   0,  36,   0,   0,  28,  77,   0,   0,   0,
+     0,  59,   0,  47,   0,   0,   0,   0,   0,   0,   0,   0,   0,
+     0,  69,   0,   0,   0,   0,   0, 144,   3,   0,  58,   0,   1,
+    75,   0,   0,   0,  31,   0,   0,   0,   0,   0,   0,  64,  66,
+    63,   0,   0,   0,   0,  46,   0,  16,   0, 115,   0,   0,   0,
+     0,   0,   0,   0,   0,   0,   0,  81,  97,   0,   8,   0, 109,
+    21,   7,  67,   0,  79,  93, 117,   0,   0,  68,   0,   0,  96,
+     0,  55,   0,  76,   0,  92,  32,  33,  57,  25,   0,  98,   0,
+     0,  87,
 };
 /* aKWLen[i] is the length (in bytes) of the i-th keyword */
-static const unsigned char aKWLen[140] = {
+static const unsigned char aKWLen[145] = {
      7,   7,   5,   4,   6,   4,   5,   3,   6,   7,   3,   6,   6,
      7,   7,   3,   8,   2,   6,   5,   4,   4,   3,  10,   4,   7,
-     6,   9,   4,   2,  10,   9,   4,   9,   4,   6,   2,   3,  11,
-     6,   2,   7,   5,   5,   6,   7,  10,   6,   5,   7,   4,   5,
-     7,   9,   6,   6,   6,   4,   5,   5,   5,   7,   7,   6,   5,
-     7,   3,   6,   4,   7,   6,  12,   9,   4,   6,   4,   5,   4,
-     7,   6,   5,   6,   6,   7,   5,   4,   7,   3,   2,   4,   5,
-     9,   5,   6,   3,   7,  13,   2,   2,   4,   6,   6,   8,   5,
-    17,  12,   7,   9,   8,   8,   2,   4,   9,   4,   6,   7,   9,
-     4,   4,   2,   6,   5,   8,   6,   4,   5,   8,   4,   3,   9,
-     5,   5,   6,   4,   6,   2,   2,   9,   3,   7,
+     6,   9,   4,   2,   6,   5,   9,   9,   4,   7,   3,   2,   4,
+     4,   6,  11,   6,   2,   7,   5,   5,   9,   6,  10,   4,   6,
+     2,   3,   7,   5,   9,   6,   6,   4,   5,   5,  10,   6,   5,
+     7,   4,   5,   7,   6,   7,   7,   6,   5,   7,   3,   7,   4,
+     7,   6,  12,   9,   4,   6,   5,   4,   7,   6,   5,   6,   6,
+     7,   6,   4,   5,   9,   5,   6,   3,   8,   8,   2,  13,   2,
+     2,   4,   6,   6,   8,   5,  17,  12,   7,   9,   4,   9,   4,
+     4,   6,   7,   5,   9,   4,   4,   5,   2,   5,   8,   6,   4,
+     5,   8,   4,   3,   9,   5,   5,   6,   4,   6,   2,   2,   9,
+     3,   7,
 };
 /* aKWOffset[i] is the index into zKWText[] of the start of
 ** the text for the i-th keyword. */
-static const unsigned short int aKWOffset[140] = {
+static const unsigned short int aKWOffset[145] = {
      0,   2,   2,   8,   9,  14,  16,  20,  23,  25,  25,  29,  33,
     36,  41,  46,  48,  53,  54,  59,  62,  65,  67,  69,  78,  81,
-    86,  90,  90,  94,  99, 106, 114, 117, 123, 126, 126, 129, 131,
-   136, 140, 141, 146, 150, 154, 159, 165, 175, 178, 183, 183, 187,
-   191, 197, 205, 211, 216, 221, 224, 227, 231, 236, 242, 248, 248,
-   254, 255, 259, 265, 269, 276, 282, 294, 303, 305, 311, 315, 320,
-   322, 329, 334, 339, 345, 351, 357, 362, 365, 365, 365, 368, 372,
-   375, 384, 388, 394, 396, 403, 405, 407, 416, 420, 426, 432, 440,
-   445, 445, 445, 461, 470, 477, 478, 485, 488, 497, 501, 506, 513,
-   522, 526, 530, 532, 538, 542, 550, 556, 559, 564, 572, 572, 576,
-   585, 590, 595, 601, 604, 607, 610, 612, 617, 621,
+    86,  90,  90,  94,  99, 101, 105, 111, 119, 123, 123, 123, 126,
+   129, 132, 137, 142, 146, 147, 152, 156, 160, 168, 174, 181, 184,
+   184, 187, 189, 195, 198, 206, 211, 216, 219, 222, 226, 236, 239,
+   244, 244, 248, 252, 259, 265, 271, 277, 277, 283, 284, 288, 295,
+   299, 306, 312, 324, 333, 335, 341, 346, 348, 355, 360, 365, 371,
+   377, 382, 388, 392, 395, 404, 408, 414, 416, 423, 424, 431, 433,
+   435, 444, 448, 454, 460, 468, 473, 473, 473, 489, 498, 501, 510,
+   513, 517, 522, 529, 534, 543, 547, 550, 555, 557, 561, 569, 575,
+   578, 583, 591, 591, 595, 604, 609, 614, 620, 623, 626, 629, 631,
+   636, 640,
 };
 /* aKWCode[i] is the parser symbol code for the i-th keyword */
-static const unsigned char aKWCode[140] = {
+static const unsigned char aKWCode[145] = {
   TK_REINDEX,    TK_INDEXED,    TK_INDEX,      TK_DESC,       TK_ESCAPE,     
   TK_EACH,       TK_CHECK,      TK_KEY,        TK_BEFORE,     TK_FOREIGN,    
   TK_FOR,        TK_IGNORE,     TK_LIKE_KW,    TK_EXPLAIN,    TK_INSTEAD,    
   TK_ADD,        TK_DATABASE,   TK_AS,         TK_SELECT,     TK_TABLE,      
   TK_JOIN_KW,    TK_THEN,       TK_END,        TK_DEFERRABLE, TK_ELSE,       
   TK_EXCLUDE,    TK_DELETE,     TK_TEMP,       TK_TEMP,       TK_OR,         
-  TK_CONSTRAINT, TK_INTERSECT,  TK_TIES,       TK_SAVEPOINT,  TK_INTO,       
-  TK_OFFSET,     TK_OF,         TK_SET,        TK_TRANSACTION,TK_ACTION,     
-  TK_ON,         TK_JOIN_KW,    TK_ALTER,      TK_RAISE,      TK_EXCEPT,     
-  TK_TRIGGER,    TK_REFERENCES, TK_UNIQUE,     TK_QUERY,      TK_WITHOUT,    
-  TK_WITH,       TK_JOIN_KW,    TK_RELEASE,    TK_EXCLUSIVE,  TK_EXISTS,     
-  TK_ATTACH,     TK_HAVING,     TK_LIKE_KW,    TK_BEGIN,      TK_JOIN_KW,    
-  TK_RANGE,      TK_BETWEEN,    TK_NOTHING,    TK_GROUPS,     TK_GROUP,      
-  TK_CASCADE,    TK_ASC,        TK_DETACH,     TK_CASE,       TK_COLLATE,    
-  TK_CREATE,     TK_CTIME_KW,   TK_IMMEDIATE,  TK_JOIN,       TK_INSERT,     
-  TK_LIKE_KW,    TK_MATCH,      TK_PLAN,       TK_ANALYZE,    TK_PRAGMA,     
-  TK_ABORT,      TK_UPDATE,     TK_VALUES,     TK_VIRTUAL,    TK_LIMIT,      
-  TK_WHEN,       TK_NOTNULL,    TK_NOT,        TK_NO,         TK_NULL,       
-  TK_WHERE,      TK_RECURSIVE,  TK_AFTER,      TK_RENAME,     TK_AND,        
-  TK_DEFAULT,    TK_AUTOINCR,   TK_TO,         TK_IN,         TK_CAST,       
-  TK_COLUMNKW,   TK_COMMIT,     TK_CONFLICT,   TK_JOIN_KW,    TK_CTIME_KW,   
-  TK_CTIME_KW,   TK_CURRENT,    TK_PARTITION,  TK_DEFERRED,   TK_DISTINCT,   
-  TK_IS,         TK_DROP,       TK_PRECEDING,  TK_FAIL,       TK_FILTER,     
-  TK_REPLACE,    TK_FOLLOWING,  TK_FROM,       TK_JOIN_KW,    TK_IF,         
-  TK_ISNULL,     TK_ORDER,      TK_RESTRICT,   TK_OTHERS,     TK_OVER,       
+  TK_ISNULL,     TK_NULLS,      TK_SAVEPOINT,  TK_INTERSECT,  TK_TIES,       
+  TK_NOTNULL,    TK_NOT,        TK_NO,         TK_NULL,       TK_LIKE_KW,    
+  TK_EXCEPT,     TK_TRANSACTION,TK_ACTION,     TK_ON,         TK_JOIN_KW,    
+  TK_ALTER,      TK_RAISE,      TK_EXCLUSIVE,  TK_EXISTS,     TK_CONSTRAINT, 
+  TK_INTO,       TK_OFFSET,     TK_OF,         TK_SET,        TK_TRIGGER,    
+  TK_RANGE,      TK_GENERATED,  TK_DETACH,     TK_HAVING,     TK_LIKE_KW,    
+  TK_BEGIN,      TK_JOIN_KW,    TK_REFERENCES, TK_UNIQUE,     TK_QUERY,      
+  TK_WITHOUT,    TK_WITH,       TK_JOIN_KW,    TK_RELEASE,    TK_ATTACH,     
+  TK_BETWEEN,    TK_NOTHING,    TK_GROUPS,     TK_GROUP,      TK_CASCADE,    
+  TK_ASC,        TK_DEFAULT,    TK_CASE,       TK_COLLATE,    TK_CREATE,     
+  TK_CTIME_KW,   TK_IMMEDIATE,  TK_JOIN,       TK_INSERT,     TK_MATCH,      
+  TK_PLAN,       TK_ANALYZE,    TK_PRAGMA,     TK_ABORT,      TK_UPDATE,     
+  TK_VALUES,     TK_VIRTUAL,    TK_ALWAYS,     TK_WHEN,       TK_WHERE,      
+  TK_RECURSIVE,  TK_AFTER,      TK_RENAME,     TK_AND,        TK_DEFERRED,   
+  TK_DISTINCT,   TK_IS,         TK_AUTOINCR,   TK_TO,         TK_IN,         
+  TK_CAST,       TK_COLUMNKW,   TK_COMMIT,     TK_CONFLICT,   TK_JOIN_KW,    
+  TK_CTIME_KW,   TK_CTIME_KW,   TK_CURRENT,    TK_PARTITION,  TK_DROP,       
+  TK_PRECEDING,  TK_FAIL,       TK_LAST,       TK_FILTER,     TK_REPLACE,    
+  TK_FIRST,      TK_FOLLOWING,  TK_FROM,       TK_JOIN_KW,    TK_LIMIT,      
+  TK_IF,         TK_ORDER,      TK_RESTRICT,   TK_OTHERS,     TK_OVER,       
   TK_JOIN_KW,    TK_ROLLBACK,   TK_ROWS,       TK_ROW,        TK_UNBOUNDED,  
   TK_UNION,      TK_USING,      TK_VACUUM,     TK_VIEW,       TK_WINDOW,     
   TK_DO,         TK_BY,         TK_INITIALLY,  TK_ALL,        TK_PRIMARY,    
 };
+/* Hash table decoded:
+**   0: INSERT
+**   1: IS
+**   2: ROLLBACK TRIGGER
+**   3: IMMEDIATE
+**   4: PARTITION
+**   5: TEMP
+**   6:
+**   7:
+**   8: VALUES WITHOUT
+**   9:
+**  10: MATCH
+**  11: NOTHING
+**  12:
+**  13: OF
+**  14: TIES IGNORE
+**  15: PLAN
+**  16: INSTEAD INDEXED
+**  17:
+**  18: TRANSACTION RIGHT
+**  19: WHEN
+**  20: SET HAVING
+**  21: IF
+**  22: ROWS
+**  23: SELECT
+**  24:
+**  25:
+**  26: VACUUM SAVEPOINT
+**  27:
+**  28: LIKE UNION VIRTUAL REFERENCES
+**  29: RESTRICT
+**  30:
+**  31: THEN REGEXP
+**  32: TO
+**  33:
+**  34: BEFORE
+**  35:
+**  36:
+**  37: FOLLOWING COLLATE CASCADE
+**  38: CREATE
+**  39:
+**  40: CASE REINDEX
+**  41: EACH
+**  42:
+**  43: QUERY
+**  44: AND ADD
+**  45: PRIMARY ANALYZE
+**  46:
+**  47: ROW ASC DETACH
+**  48: CURRENT_TIME CURRENT_DATE
+**  49:
+**  50:
+**  51: EXCLUSIVE TEMPORARY
+**  52:
+**  53: DEFERRED
+**  54: DEFERRABLE
+**  55:
+**  56: DATABASE
+**  57:
+**  58: DELETE VIEW GENERATED
+**  59: ATTACH
+**  60: END
+**  61: EXCLUDE
+**  62: ESCAPE DESC
+**  63: GLOB
+**  64: WINDOW ELSE
+**  65: COLUMN
+**  66: FIRST
+**  67:
+**  68: GROUPS ALL
+**  69: DISTINCT DROP KEY
+**  70: BETWEEN
+**  71: INITIALLY
+**  72: BEGIN
+**  73: FILTER CHECK ACTION
+**  74: GROUP INDEX
+**  75:
+**  76: EXISTS DEFAULT
+**  77:
+**  78: FOR CURRENT_TIMESTAMP
+**  79: EXCEPT
+**  80:
+**  81: CROSS
+**  82:
+**  83:
+**  84:
+**  85: CAST
+**  86: FOREIGN AUTOINCREMENT
+**  87: COMMIT
+**  88: CURRENT AFTER ALTER
+**  89: FULL FAIL CONFLICT
+**  90: EXPLAIN
+**  91: CONSTRAINT
+**  92: FROM ALWAYS
+**  93:
+**  94: ABORT
+**  95:
+**  96: AS DO
+**  97: REPLACE WITH RELEASE
+**  98: BY RENAME
+**  99: RANGE RAISE
+** 100: OTHERS
+** 101: USING NULLS
+** 102: PRAGMA
+** 103: JOIN ISNULL OFFSET
+** 104: NOT
+** 105: OR LAST LEFT
+** 106: LIMIT
+** 107:
+** 108:
+** 109: IN
+** 110: INTO
+** 111: OVER RECURSIVE
+** 112: ORDER OUTER
+** 113:
+** 114: INTERSECT UNBOUNDED
+** 115:
+** 116:
+** 117: ON
+** 118:
+** 119: WHERE
+** 120: NO INNER
+** 121: NULL
+** 122:
+** 123: TABLE
+** 124: NATURAL NOTNULL
+** 125: PRECEDING
+** 126: UPDATE UNIQUE
+*/
 /* Check to see if z[0..n-1] is a keyword. If it is, write the
 ** parser symbol code for that keyword into *pType.  Always
 ** return the integer n (the length of the token). */
@@ -153369,116 +157806,121 @@ static int keywordCode(const char *z, int n, int *pType){
       testcase( i==27 ); /* TEMPORARY */
       testcase( i==28 ); /* TEMP */
       testcase( i==29 ); /* OR */
-      testcase( i==30 ); /* CONSTRAINT */
-      testcase( i==31 ); /* INTERSECT */
-      testcase( i==32 ); /* TIES */
-      testcase( i==33 ); /* SAVEPOINT */
-      testcase( i==34 ); /* INTO */
-      testcase( i==35 ); /* OFFSET */
-      testcase( i==36 ); /* OF */
-      testcase( i==37 ); /* SET */
-      testcase( i==38 ); /* TRANSACTION */
-      testcase( i==39 ); /* ACTION */
-      testcase( i==40 ); /* ON */
-      testcase( i==41 ); /* NATURAL */
-      testcase( i==42 ); /* ALTER */
-      testcase( i==43 ); /* RAISE */
-      testcase( i==44 ); /* EXCEPT */
-      testcase( i==45 ); /* TRIGGER */
-      testcase( i==46 ); /* REFERENCES */
-      testcase( i==47 ); /* UNIQUE */
-      testcase( i==48 ); /* QUERY */
-      testcase( i==49 ); /* WITHOUT */
-      testcase( i==50 ); /* WITH */
-      testcase( i==51 ); /* OUTER */
-      testcase( i==52 ); /* RELEASE */
-      testcase( i==53 ); /* EXCLUSIVE */
-      testcase( i==54 ); /* EXISTS */
-      testcase( i==55 ); /* ATTACH */
-      testcase( i==56 ); /* HAVING */
-      testcase( i==57 ); /* GLOB */
-      testcase( i==58 ); /* BEGIN */
-      testcase( i==59 ); /* INNER */
-      testcase( i==60 ); /* RANGE */
-      testcase( i==61 ); /* BETWEEN */
-      testcase( i==62 ); /* NOTHING */
-      testcase( i==63 ); /* GROUPS */
-      testcase( i==64 ); /* GROUP */
-      testcase( i==65 ); /* CASCADE */
-      testcase( i==66 ); /* ASC */
-      testcase( i==67 ); /* DETACH */
-      testcase( i==68 ); /* CASE */
-      testcase( i==69 ); /* COLLATE */
-      testcase( i==70 ); /* CREATE */
-      testcase( i==71 ); /* CURRENT_DATE */
-      testcase( i==72 ); /* IMMEDIATE */
-      testcase( i==73 ); /* JOIN */
-      testcase( i==74 ); /* INSERT */
-      testcase( i==75 ); /* LIKE */
-      testcase( i==76 ); /* MATCH */
-      testcase( i==77 ); /* PLAN */
-      testcase( i==78 ); /* ANALYZE */
-      testcase( i==79 ); /* PRAGMA */
-      testcase( i==80 ); /* ABORT */
-      testcase( i==81 ); /* UPDATE */
-      testcase( i==82 ); /* VALUES */
-      testcase( i==83 ); /* VIRTUAL */
-      testcase( i==84 ); /* LIMIT */
-      testcase( i==85 ); /* WHEN */
-      testcase( i==86 ); /* NOTNULL */
-      testcase( i==87 ); /* NOT */
-      testcase( i==88 ); /* NO */
-      testcase( i==89 ); /* NULL */
-      testcase( i==90 ); /* WHERE */
-      testcase( i==91 ); /* RECURSIVE */
-      testcase( i==92 ); /* AFTER */
-      testcase( i==93 ); /* RENAME */
-      testcase( i==94 ); /* AND */
-      testcase( i==95 ); /* DEFAULT */
-      testcase( i==96 ); /* AUTOINCREMENT */
-      testcase( i==97 ); /* TO */
-      testcase( i==98 ); /* IN */
-      testcase( i==99 ); /* CAST */
-      testcase( i==100 ); /* COLUMN */
-      testcase( i==101 ); /* COMMIT */
-      testcase( i==102 ); /* CONFLICT */
-      testcase( i==103 ); /* CROSS */
-      testcase( i==104 ); /* CURRENT_TIMESTAMP */
-      testcase( i==105 ); /* CURRENT_TIME */
-      testcase( i==106 ); /* CURRENT */
-      testcase( i==107 ); /* PARTITION */
-      testcase( i==108 ); /* DEFERRED */
-      testcase( i==109 ); /* DISTINCT */
-      testcase( i==110 ); /* IS */
-      testcase( i==111 ); /* DROP */
-      testcase( i==112 ); /* PRECEDING */
-      testcase( i==113 ); /* FAIL */
-      testcase( i==114 ); /* FILTER */
-      testcase( i==115 ); /* REPLACE */
-      testcase( i==116 ); /* FOLLOWING */
-      testcase( i==117 ); /* FROM */
-      testcase( i==118 ); /* FULL */
-      testcase( i==119 ); /* IF */
-      testcase( i==120 ); /* ISNULL */
-      testcase( i==121 ); /* ORDER */
-      testcase( i==122 ); /* RESTRICT */
-      testcase( i==123 ); /* OTHERS */
-      testcase( i==124 ); /* OVER */
-      testcase( i==125 ); /* RIGHT */
-      testcase( i==126 ); /* ROLLBACK */
-      testcase( i==127 ); /* ROWS */
-      testcase( i==128 ); /* ROW */
-      testcase( i==129 ); /* UNBOUNDED */
-      testcase( i==130 ); /* UNION */
-      testcase( i==131 ); /* USING */
-      testcase( i==132 ); /* VACUUM */
-      testcase( i==133 ); /* VIEW */
-      testcase( i==134 ); /* WINDOW */
-      testcase( i==135 ); /* DO */
-      testcase( i==136 ); /* BY */
-      testcase( i==137 ); /* INITIALLY */
-      testcase( i==138 ); /* ALL */
-      testcase( i==139 ); /* PRIMARY */
+      testcase( i==30 ); /* ISNULL */
+      testcase( i==31 ); /* NULLS */
+      testcase( i==32 ); /* SAVEPOINT */
+      testcase( i==33 ); /* INTERSECT */
+      testcase( i==34 ); /* TIES */
+      testcase( i==35 ); /* NOTNULL */
+      testcase( i==36 ); /* NOT */
+      testcase( i==37 ); /* NO */
+      testcase( i==38 ); /* NULL */
+      testcase( i==39 ); /* LIKE */
+      testcase( i==40 ); /* EXCEPT */
+      testcase( i==41 ); /* TRANSACTION */
+      testcase( i==42 ); /* ACTION */
+      testcase( i==43 ); /* ON */
+      testcase( i==44 ); /* NATURAL */
+      testcase( i==45 ); /* ALTER */
+      testcase( i==46 ); /* RAISE */
+      testcase( i==47 ); /* EXCLUSIVE */
+      testcase( i==48 ); /* EXISTS */
+      testcase( i==49 ); /* CONSTRAINT */
+      testcase( i==50 ); /* INTO */
+      testcase( i==51 ); /* OFFSET */
+      testcase( i==52 ); /* OF */
+      testcase( i==53 ); /* SET */
+      testcase( i==54 ); /* TRIGGER */
+      testcase( i==55 ); /* RANGE */
+      testcase( i==56 ); /* GENERATED */
+      testcase( i==57 ); /* DETACH */
+      testcase( i==58 ); /* HAVING */
+      testcase( i==59 ); /* GLOB */
+      testcase( i==60 ); /* BEGIN */
+      testcase( i==61 ); /* INNER */
+      testcase( i==62 ); /* REFERENCES */
+      testcase( i==63 ); /* UNIQUE */
+      testcase( i==64 ); /* QUERY */
+      testcase( i==65 ); /* WITHOUT */
+      testcase( i==66 ); /* WITH */
+      testcase( i==67 ); /* OUTER */
+      testcase( i==68 ); /* RELEASE */
+      testcase( i==69 ); /* ATTACH */
+      testcase( i==70 ); /* BETWEEN */
+      testcase( i==71 ); /* NOTHING */
+      testcase( i==72 ); /* GROUPS */
+      testcase( i==73 ); /* GROUP */
+      testcase( i==74 ); /* CASCADE */
+      testcase( i==75 ); /* ASC */
+      testcase( i==76 ); /* DEFAULT */
+      testcase( i==77 ); /* CASE */
+      testcase( i==78 ); /* COLLATE */
+      testcase( i==79 ); /* CREATE */
+      testcase( i==80 ); /* CURRENT_DATE */
+      testcase( i==81 ); /* IMMEDIATE */
+      testcase( i==82 ); /* JOIN */
+      testcase( i==83 ); /* INSERT */
+      testcase( i==84 ); /* MATCH */
+      testcase( i==85 ); /* PLAN */
+      testcase( i==86 ); /* ANALYZE */
+      testcase( i==87 ); /* PRAGMA */
+      testcase( i==88 ); /* ABORT */
+      testcase( i==89 ); /* UPDATE */
+      testcase( i==90 ); /* VALUES */
+      testcase( i==91 ); /* VIRTUAL */
+      testcase( i==92 ); /* ALWAYS */
+      testcase( i==93 ); /* WHEN */
+      testcase( i==94 ); /* WHERE */
+      testcase( i==95 ); /* RECURSIVE */
+      testcase( i==96 ); /* AFTER */
+      testcase( i==97 ); /* RENAME */
+      testcase( i==98 ); /* AND */
+      testcase( i==99 ); /* DEFERRED */
+      testcase( i==100 ); /* DISTINCT */
+      testcase( i==101 ); /* IS */
+      testcase( i==102 ); /* AUTOINCREMENT */
+      testcase( i==103 ); /* TO */
+      testcase( i==104 ); /* IN */
+      testcase( i==105 ); /* CAST */
+      testcase( i==106 ); /* COLUMN */
+      testcase( i==107 ); /* COMMIT */
+      testcase( i==108 ); /* CONFLICT */
+      testcase( i==109 ); /* CROSS */
+      testcase( i==110 ); /* CURRENT_TIMESTAMP */
+      testcase( i==111 ); /* CURRENT_TIME */
+      testcase( i==112 ); /* CURRENT */
+      testcase( i==113 ); /* PARTITION */
+      testcase( i==114 ); /* DROP */
+      testcase( i==115 ); /* PRECEDING */
+      testcase( i==116 ); /* FAIL */
+      testcase( i==117 ); /* LAST */
+      testcase( i==118 ); /* FILTER */
+      testcase( i==119 ); /* REPLACE */
+      testcase( i==120 ); /* FIRST */
+      testcase( i==121 ); /* FOLLOWING */
+      testcase( i==122 ); /* FROM */
+      testcase( i==123 ); /* FULL */
+      testcase( i==124 ); /* LIMIT */
+      testcase( i==125 ); /* IF */
+      testcase( i==126 ); /* ORDER */
+      testcase( i==127 ); /* RESTRICT */
+      testcase( i==128 ); /* OTHERS */
+      testcase( i==129 ); /* OVER */
+      testcase( i==130 ); /* RIGHT */
+      testcase( i==131 ); /* ROLLBACK */
+      testcase( i==132 ); /* ROWS */
+      testcase( i==133 ); /* ROW */
+      testcase( i==134 ); /* UNBOUNDED */
+      testcase( i==135 ); /* UNION */
+      testcase( i==136 ); /* USING */
+      testcase( i==137 ); /* VACUUM */
+      testcase( i==138 ); /* VIEW */
+      testcase( i==139 ); /* WINDOW */
+      testcase( i==140 ); /* DO */
+      testcase( i==141 ); /* BY */
+      testcase( i==142 ); /* INITIALLY */
+      testcase( i==143 ); /* ALL */
+      testcase( i==144 ); /* PRIMARY */
       *pType = aKWCode[i];
       break;
     }
@@ -153490,7 +157932,7 @@ SQLITE_PRIVATE int sqlite3KeywordCode(const unsigned char *z, int n){
   keywordCode((char*)z, n, &id);
   return id;
 }
-#define SQLITE_N_KEYWORD 140
+#define SQLITE_N_KEYWORD 145
 SQLITE_API int sqlite3_keyword_name(int i,const char **pzName,int *pnName){
   if( i<0 || i>=SQLITE_N_KEYWORD ) return SQLITE_ERROR;
   *pzName = zKWText + aKWOffset[i];
@@ -154117,7 +158559,7 @@ SQLITE_PRIVATE char *sqlite3Normalize(
   int nParen;        /* Number of nested levels of parentheses */
   int iStartIN;      /* Start of RHS of IN operator in z[] */
   int nParenAtIN;    /* Value of nParent at start of RHS of IN operator */
-  int j;             /* Bytes of normalized SQL generated so far */
+  u32 j;             /* Bytes of normalized SQL generated so far */
   sqlite3_str *pStr; /* The normalized SQL string under construction */
 
   db = sqlite3VdbeDb(pVdbe);
@@ -154161,7 +158603,7 @@ SQLITE_PRIVATE char *sqlite3Normalize(
       }
       case TK_RP: {
         if( iStartIN>0 && nParen==nParenAtIN ){
-          assert( pStr->nChar>=iStartIN );
+          assert( pStr->nChar>=(u32)iStartIN );
           pStr->nChar = iStartIN+1;
           sqlite3_str_append(pStr, "?,?,?", 5);
           iStartIN = 0;
@@ -155298,6 +159740,9 @@ SQLITE_API int sqlite3_config(int op, ...){
 static int setupLookaside(sqlite3 *db, void *pBuf, int sz, int cnt){
 #ifndef SQLITE_OMIT_LOOKASIDE
   void *pStart;
+  sqlite3_int64 szAlloc = sz*(sqlite3_int64)cnt;
+  int nBig;   /* Number of full-size slots */
+  int nSm;    /* Number smaller LOOKASIDE_SMALL-byte slots */
   
   if( sqlite3LookasideUsed(db,0)>0 ){
     return SQLITE_BUSY;
@@ -155320,37 +159765,71 @@ static int setupLookaside(sqlite3 *db, void *pBuf, int sz, int cnt){
     pStart = 0;
   }else if( pBuf==0 ){
     sqlite3BeginBenignMalloc();
-    pStart = sqlite3Malloc( sz*(sqlite3_int64)cnt );  /* IMP: R-61949-35727 */
+    pStart = sqlite3Malloc( szAlloc );  /* IMP: R-61949-35727 */
     sqlite3EndBenignMalloc();
-    if( pStart ) cnt = sqlite3MallocSize(pStart)/sz;
+    if( pStart ) szAlloc = sqlite3MallocSize(pStart);
   }else{
     pStart = pBuf;
   }
+#ifndef SQLITE_OMIT_TWOSIZE_LOOKASIDE
+  if( sz>=LOOKASIDE_SMALL*3 ){
+    nBig = szAlloc/(3*LOOKASIDE_SMALL+sz);
+    nSm = (szAlloc - sz*nBig)/LOOKASIDE_SMALL;
+  }else if( sz>=LOOKASIDE_SMALL*2 ){
+    nBig = szAlloc/(LOOKASIDE_SMALL+sz);
+    nSm = (szAlloc - sz*nBig)/LOOKASIDE_SMALL;
+  }else
+#endif /* SQLITE_OMIT_TWOSIZE_LOOKASIDE */
+  if( sz>0 ){
+    nBig = szAlloc/sz;
+    nSm = 0;
+  }else{
+    nBig = nSm = 0;
+  }
   db->lookaside.pStart = pStart;
   db->lookaside.pInit = 0;
   db->lookaside.pFree = 0;
   db->lookaside.sz = (u16)sz;
+  db->lookaside.szTrue = (u16)sz;
   if( pStart ){
     int i;
     LookasideSlot *p;
     assert( sz > (int)sizeof(LookasideSlot*) );
-    db->lookaside.nSlot = cnt;
     p = (LookasideSlot*)pStart;
-    for(i=cnt-1; i>=0; i--){
+    for(i=0; i<nBig; i++){
       p->pNext = db->lookaside.pInit;
       db->lookaside.pInit = p;
       p = (LookasideSlot*)&((u8*)p)[sz];
     }
+#ifndef SQLITE_OMIT_TWOSIZE_LOOKASIDE
+    db->lookaside.pSmallInit = 0;
+    db->lookaside.pSmallFree = 0;
+    db->lookaside.pMiddle = p;
+    for(i=0; i<nSm; i++){
+      p->pNext = db->lookaside.pSmallInit;
+      db->lookaside.pSmallInit = p;
+      p = (LookasideSlot*)&((u8*)p)[LOOKASIDE_SMALL];
+    }
+#endif /* SQLITE_OMIT_TWOSIZE_LOOKASIDE */
+    assert( ((uptr)p)<=szAlloc + (uptr)pStart );
     db->lookaside.pEnd = p;
     db->lookaside.bDisable = 0;
     db->lookaside.bMalloced = pBuf==0 ?1:0;
+    db->lookaside.nSlot = nBig+nSm;
   }else{
     db->lookaside.pStart = db;
+#ifndef SQLITE_OMIT_TWOSIZE_LOOKASIDE
+    db->lookaside.pSmallInit = 0;
+    db->lookaside.pSmallFree = 0;
+    db->lookaside.pMiddle = db;
+#endif /* SQLITE_OMIT_TWOSIZE_LOOKASIDE */
     db->lookaside.pEnd = db;
     db->lookaside.bDisable = 1;
+    db->lookaside.sz = 0;
     db->lookaside.bMalloced = 0;
     db->lookaside.nSlot = 0;
   }
+  assert( sqlite3LookasideUsed(db,0)==0 );
 #endif /* SQLITE_OMIT_LOOKASIDE */
   return SQLITE_OK;
 }
@@ -155451,6 +159930,7 @@ SQLITE_API int sqlite3_db_config(sqlite3 *db, int op, ...){
       } aFlagOp[] = {
         { SQLITE_DBCONFIG_ENABLE_FKEY,           SQLITE_ForeignKeys    },
         { SQLITE_DBCONFIG_ENABLE_TRIGGER,        SQLITE_EnableTrigger  },
+        { SQLITE_DBCONFIG_ENABLE_VIEW,           SQLITE_EnableView     },
         { SQLITE_DBCONFIG_ENABLE_FTS3_TOKENIZER, SQLITE_Fts3Tokenizer  },
         { SQLITE_DBCONFIG_ENABLE_LOAD_EXTENSION, SQLITE_LoadExtension  },
         { SQLITE_DBCONFIG_NO_CKPT_ON_CLOSE,      SQLITE_NoCkptOnClose  },
@@ -155460,6 +159940,11 @@ SQLITE_API int sqlite3_db_config(sqlite3 *db, int op, ...){
         { SQLITE_DBCONFIG_DEFENSIVE,             SQLITE_Defensive      },
         { SQLITE_DBCONFIG_WRITABLE_SCHEMA,       SQLITE_WriteSchema|
                                                  SQLITE_NoSchemaError  },
+        { SQLITE_DBCONFIG_LEGACY_ALTER_TABLE,    SQLITE_LegacyAlter    },
+        { SQLITE_DBCONFIG_DQS_DDL,               SQLITE_DqsDDL         },
+        { SQLITE_DBCONFIG_DQS_DML,               SQLITE_DqsDML         },
+        { SQLITE_DBCONFIG_LEGACY_FILE_FORMAT,    SQLITE_LegacyFileFmt  },
+        { SQLITE_DBCONFIG_TRUSTED_SCHEMA,        SQLITE_TrustedSchema  },
       };
       unsigned int i;
       rc = SQLITE_ERROR; /* IMP: R-42790-23372 */
@@ -155490,28 +159975,17 @@ SQLITE_API int sqlite3_db_config(sqlite3 *db, int op, ...){
   return rc;
 }
 
-
-/*
-** Return true if the buffer z[0..n-1] contains all spaces.
-*/
-static int allSpaces(const char *z, int n){
-  while( n>0 && z[n-1]==' ' ){ n--; }
-  return n==0;
-}
-
 /*
 ** This is the default collating function named "BINARY" which is always
 ** available.
-**
-** If the padFlag argument is not NULL then space padding at the end
-** of strings is ignored.  This implements the RTRIM collation.
 */
 static int binCollFunc(
-  void *padFlag,
+  void *NotUsed,
   int nKey1, const void *pKey1,
   int nKey2, const void *pKey2
 ){
   int rc, n;
+  UNUSED_PARAMETER(NotUsed);
   n = nKey1<nKey2 ? nKey1 : nKey2;
   /* EVIDENCE-OF: R-65033-28449 The built-in BINARY collation compares
   ** strings byte by byte using the memcmp() function from the standard C
@@ -155519,29 +159993,33 @@ static int binCollFunc(
   assert( pKey1 && pKey2 );
   rc = memcmp(pKey1, pKey2, n);
   if( rc==0 ){
-    if( padFlag
-     && allSpaces(((char*)pKey1)+n, nKey1-n)
-     && allSpaces(((char*)pKey2)+n, nKey2-n)
-    ){
-      /* EVIDENCE-OF: R-31624-24737 RTRIM is like BINARY except that extra
-      ** spaces at the end of either string do not change the result. In other
-      ** words, strings will compare equal to one another as long as they
-      ** differ only in the number of spaces at the end.
-      */
-    }else{
-      rc = nKey1 - nKey2;
-    }
+    rc = nKey1 - nKey2;
   }
   return rc;
 }
 
+/*
+** This is the collating function named "RTRIM" which is always
+** available.  Ignore trailing spaces.
+*/
+static int rtrimCollFunc(
+  void *pUser,
+  int nKey1, const void *pKey1,
+  int nKey2, const void *pKey2
+){
+  const u8 *pK1 = (const u8*)pKey1;
+  const u8 *pK2 = (const u8*)pKey2;
+  while( nKey1 && pK1[nKey1-1]==' ' ) nKey1--;
+  while( nKey2 && pK2[nKey2-1]==' ' ) nKey2--;
+  return binCollFunc(pUser, nKey1, pKey1, nKey2, pKey2);
+}
+
 /*
 ** Return true if CollSeq is the default built-in BINARY.
 */
 SQLITE_PRIVATE int sqlite3IsBinary(const CollSeq *p){
-  assert( p==0 || p->xCmp!=binCollFunc || p->pUser!=0
-            || strcmp(p->zName,"BINARY")==0 );
-  return p==0 || (p->xCmp==binCollFunc && p->pUser==0);
+  assert( p==0 || p->xCmp!=binCollFunc || strcmp(p->zName,"BINARY")==0 );
+  return p==0 || p->xCmp==binCollFunc;
 }
 
 /*
@@ -155854,11 +160332,8 @@ SQLITE_PRIVATE void sqlite3LeaveMutexAndCloseZombie(sqlite3 *db){
 #ifndef SQLITE_OMIT_VIRTUALTABLE
   for(i=sqliteHashFirst(&db->aModule); i; i=sqliteHashNext(i)){
     Module *pMod = (Module *)sqliteHashData(i);
-    if( pMod->xDestroy ){
-      pMod->xDestroy(pMod->pAux);
-    }
     sqlite3VtabEponymousTableClear(db, pMod);
-    sqlite3DbFree(db, pMod);
+    sqlite3VtabModuleUnref(db, pMod);
   }
   sqlite3HashClear(&db->aModule);
 #endif
@@ -156008,6 +160483,7 @@ SQLITE_PRIVATE const char *sqlite3ErrName(int rc){
       case SQLITE_CANTOPEN_ISDIR:     zName = "SQLITE_CANTOPEN_ISDIR";    break;
       case SQLITE_CANTOPEN_FULLPATH:  zName = "SQLITE_CANTOPEN_FULLPATH"; break;
       case SQLITE_CANTOPEN_CONVPATH:  zName = "SQLITE_CANTOPEN_CONVPATH"; break;
+      case SQLITE_CANTOPEN_SYMLINK:   zName = "SQLITE_CANTOPEN_SYMLINK";  break;
       case SQLITE_PROTOCOL:           zName = "SQLITE_PROTOCOL";          break;
       case SQLITE_EMPTY:              zName = "SQLITE_EMPTY";             break;
       case SQLITE_SCHEMA:             zName = "SQLITE_SCHEMA";            break;
@@ -156339,8 +160815,16 @@ SQLITE_PRIVATE int sqlite3CreateFunc(
   }
 
   assert( SQLITE_FUNC_CONSTANT==SQLITE_DETERMINISTIC );
-  extraFlags = enc &  SQLITE_DETERMINISTIC;
+  assert( SQLITE_FUNC_DIRECT==SQLITE_DIRECTONLY );
+  extraFlags = enc &  (SQLITE_DETERMINISTIC|SQLITE_DIRECTONLY|
+                       SQLITE_SUBTYPE|SQLITE_INNOCUOUS);
   enc &= (SQLITE_FUNC_ENCMASK|SQLITE_ANY);
+
+  /* The SQLITE_INNOCUOUS flag is the same bit as SQLITE_FUNC_UNSAFE.  But
+  ** the meaning is inverted.  So flip the bit. */
+  assert( SQLITE_FUNC_UNSAFE==SQLITE_INNOCUOUS );
+  extraFlags ^= SQLITE_FUNC_UNSAFE;
+
   
 #ifndef SQLITE_OMIT_UTF16
   /* If SQLITE_UTF16 is specified as the encoding type, transform this
@@ -156354,11 +160838,13 @@ SQLITE_PRIVATE int sqlite3CreateFunc(
     enc = SQLITE_UTF16NATIVE;
   }else if( enc==SQLITE_ANY ){
     int rc;
-    rc = sqlite3CreateFunc(db, zFunctionName, nArg, SQLITE_UTF8|extraFlags,
+    rc = sqlite3CreateFunc(db, zFunctionName, nArg,
+         (SQLITE_UTF8|extraFlags)^SQLITE_FUNC_UNSAFE,
          pUserData, xSFunc, xStep, xFinal, xValue, xInverse, pDestructor);
     if( rc==SQLITE_OK ){
-      rc = sqlite3CreateFunc(db, zFunctionName, nArg, SQLITE_UTF16LE|extraFlags,
-          pUserData, xSFunc, xStep, xFinal, xValue, xInverse, pDestructor);
+      rc = sqlite3CreateFunc(db, zFunctionName, nArg,
+           (SQLITE_UTF16LE|extraFlags)^SQLITE_FUNC_UNSAFE,
+           pUserData, xSFunc, xStep, xFinal, xValue, xInverse, pDestructor);
     }
     if( rc!=SQLITE_OK ){
       return rc;
@@ -156402,6 +160888,7 @@ SQLITE_PRIVATE int sqlite3CreateFunc(
   p->u.pDestructor = pDestructor;
   p->funcFlags = (p->funcFlags & SQLITE_FUNC_ENCMASK) | extraFlags;
   testcase( p->funcFlags & SQLITE_DETERMINISTIC );
+  testcase( p->funcFlags & SQLITE_DIRECTONLY );
   p->xSFunc = xSFunc ? xSFunc : xStep;
   p->xFinalize = xFinal;
   p->xValue = xValue;
@@ -157683,6 +162170,7 @@ static int openDatabase(
   db->magic = SQLITE_MAGIC_BUSY;
   db->aDb = db->aDbStatic;
   db->lookaside.bDisable = 1;
+  db->lookaside.sz = 0;
 
   assert( sizeof(db->aLimit)==sizeof(aHardLimit) );
   memcpy(db->aLimit, aHardLimit, sizeof(db->aLimit));
@@ -157692,7 +162180,38 @@ static int openDatabase(
   db->szMmap = sqlite3GlobalConfig.szMmap;
   db->nextPagesize = 0;
   db->nMaxSorterMmap = 0x7FFFFFFF;
-  db->flags |= SQLITE_ShortColNames | SQLITE_EnableTrigger | SQLITE_CacheSpill
+  db->flags |= SQLITE_ShortColNames
+                 | SQLITE_EnableTrigger
+                 | SQLITE_EnableView
+                 | SQLITE_CacheSpill
+#if !defined(SQLITE_TRUSTED_SCHEMA) || SQLITE_TRUSTED_SCHEMA+0!=0
+                 | SQLITE_TrustedSchema
+#endif
+/* The SQLITE_DQS compile-time option determines the default settings
+** for SQLITE_DBCONFIG_DQS_DDL and SQLITE_DBCONFIG_DQS_DML.
+**
+**    SQLITE_DQS     SQLITE_DBCONFIG_DQS_DDL    SQLITE_DBCONFIG_DQS_DML
+**    ----------     -----------------------    -----------------------
+**     undefined               on                          on   
+**         3                   on                          on
+**         2                   on                         off
+**         1                  off                          on
+**         0                  off                         off
+**
+** Legacy behavior is 3 (double-quoted string literals are allowed anywhere)
+** and so that is the default.  But developers are encouranged to use
+** -DSQLITE_DQS=0 (best) or -DSQLITE_DQS=1 (second choice) if possible.
+*/
+#if !defined(SQLITE_DQS)
+# define SQLITE_DQS 3
+#endif
+#if (SQLITE_DQS&1)==1
+                 | SQLITE_DqsDML
+#endif
+#if (SQLITE_DQS&2)==2
+                 | SQLITE_DqsDDL
+#endif
+
 #if !defined(SQLITE_DEFAULT_AUTOMATIC_INDEX) || SQLITE_DEFAULT_AUTOMATIC_INDEX
                  | SQLITE_AutoIndex
 #endif
@@ -157743,7 +162262,7 @@ static int openDatabase(
   createCollation(db, sqlite3StrBINARY, SQLITE_UTF16BE, 0, binCollFunc, 0);
   createCollation(db, sqlite3StrBINARY, SQLITE_UTF16LE, 0, binCollFunc, 0);
   createCollation(db, "NOCASE", SQLITE_UTF8, 0, nocaseCollatingFunc, 0);
-  createCollation(db, "RTRIM", SQLITE_UTF8, (void*)1, binCollFunc, 0);
+  createCollation(db, "RTRIM", SQLITE_UTF8, 0, rtrimCollFunc, 0);
   if( db->mallocFailed ){
     goto opendb_out;
   }
@@ -157897,6 +162416,13 @@ static int openDatabase(
   }
 #endif
 
+#ifdef SQLITE_ENABLE_INTERNAL_FUNCTIONS
+  /* Testing use only!!! The -DSQLITE_ENABLE_INTERNAL_FUNCTIONS=1 compile-time
+  ** option gives access to internal functions by default.  
+  ** Testing use only!!! */
+  db->mDbFlags |= DBFLAG_InternalFunc;
+#endif
+
   /* -DSQLITE_DEFAULT_LOCKING_MODE=1 makes EXCLUSIVE the default locking
   ** mode.  -DSQLITE_DEFAULT_LOCKING_MODE=0 make NORMAL the default locking
   ** mode.  Doing nothing at all also makes NORMAL the default.
@@ -158415,15 +162941,38 @@ SQLITE_API int sqlite3_test_control(int op, ...){
       break;
     }
 
-    /*
-    ** Reset the PRNG back to its uninitialized state.  The next call
-    ** to sqlite3_randomness() will reseed the PRNG using a single call
-    ** to the xRandomness method of the default VFS.
+    /*  sqlite3_test_control(SQLITE_TESTCTRL_PRNG_SEED, int x, sqlite3 *db);
+    **
+    ** Control the seed for the pseudo-random number generator (PRNG) that
+    ** is built into SQLite.  Cases:
+    **
+    **    x!=0 && db!=0       Seed the PRNG to the current value of the
+    **                        schema cookie in the main database for db, or
+    **                        x if the schema cookie is zero.  This case
+    **                        is convenient to use with database fuzzers
+    **                        as it allows the fuzzer some control over the
+    **                        the PRNG seed.
+    **
+    **    x!=0 && db==0       Seed the PRNG to the value of x.
+    **
+    **    x==0 && db==0       Revert to default behavior of using the
+    **                        xRandomness method on the primary VFS.
+    **
+    ** This test-control also resets the PRNG so that the new seed will
+    ** be used for the next call to sqlite3_randomness().
     */
-    case SQLITE_TESTCTRL_PRNG_RESET: {
+#ifndef SQLITE_OMIT_WSD
+    case SQLITE_TESTCTRL_PRNG_SEED: {
+      int x = va_arg(ap, int);
+      int y;
+      sqlite3 *db = va_arg(ap, sqlite3*);
+      assert( db==0 || db->aDb[0].pSchema!=0 );
+      if( db && (y = db->aDb[0].pSchema->schema_cookie)!=0 ){ x = y; }
+      sqlite3Config.iPrngSeed = x;
       sqlite3_randomness(0,0);
       break;
     }
+#endif
 
     /*
     **  sqlite3_test_control(BITVEC_TEST, size, program)
@@ -158608,15 +163157,14 @@ SQLITE_API int sqlite3_test_control(int op, ...){
       break;
     }
 
-    /*   sqlite3_test_control(SQLITE_TESTCTRL_INTERNAL_FUNCS, int onoff);
+    /*   sqlite3_test_control(SQLITE_TESTCTRL_INTERNAL_FUNCTIONS, sqlite3*);
     **
-    ** If parameter onoff is non-zero, internal-use-only SQL functions
-    ** are visible to ordinary SQL.  This is useful for testing but is
-    ** unsafe because invalid parameters to those internal-use-only functions
-    ** can result in crashes or segfaults.
+    ** Toggle the ability to use internal functions on or off for
+    ** the database connection given in the argument.
     */
     case SQLITE_TESTCTRL_INTERNAL_FUNCTIONS: {
-      sqlite3GlobalConfig.bInternalFunctions = va_arg(ap, int);
+      sqlite3 *db = va_arg(ap, sqlite3*);
+      db->mDbFlags ^= DBFLAG_InternalFunc;
       break;
     }
 
@@ -158633,6 +163181,17 @@ SQLITE_API int sqlite3_test_control(int op, ...){
       break;
     }
 
+    /*   sqlite3_test_control(SQLITE_TESTCTRL_EXTRA_SCHEMA_CHECKS, int);
+    **
+    ** Set or clear a flag that causes SQLite to verify that type, name,
+    ** and tbl_name fields of the sqlite_master table.  This is normally
+    ** on, but it is sometimes useful to turn it off for testing.
+    */
+    case SQLITE_TESTCTRL_EXTRA_SCHEMA_CHECKS: {
+      sqlite3GlobalConfig.bExtraSchemaChecks = va_arg(ap, int);
+      break;
+    }
+
     /* Set the threshold at which OP_Once counters reset back to zero.
     ** By default this is 0x7ffffffe (over 2 billion), but that value is
     ** too big to test in a reasonable amount of time, so this control is
@@ -158719,12 +163278,43 @@ SQLITE_API int sqlite3_test_control(int op, ...){
       break;
     }
 #endif /* defined(YYCOVERAGE) */
+
+    /*  sqlite3_test_control(SQLITE_TESTCTRL_RESULT_INTREAL, sqlite3_context*);
+    **
+    ** This test-control causes the most recent sqlite3_result_int64() value
+    ** to be interpreted as a MEM_IntReal instead of as an MEM_Int.  Normally,
+    ** MEM_IntReal values only arise during an INSERT operation of integer
+    ** values into a REAL column, so they can be challenging to test.  This
+    ** test-control enables us to write an intreal() SQL function that can
+    ** inject an intreal() value at arbitrary places in an SQL statement,
+    ** for testing purposes.
+    */
+    case SQLITE_TESTCTRL_RESULT_INTREAL: {
+      sqlite3_context *pCtx = va_arg(ap, sqlite3_context*);
+      sqlite3ResultIntReal(pCtx);
+      break;
+    }
   }
   va_end(ap);
 #endif /* SQLITE_UNTESTABLE */
   return rc;
 }
 
+/*
+** The Pager stores the Database filename, Journal filename, and WAL filename
+** consecutively in memory, in that order.  The database filename is prefixed
+** by four zero bytes.  Locate the start of the database filename by searching
+** backwards for the first byte following four consecutive zero bytes.
+**
+** This only works if the filename passed in was obtained from the Pager.
+*/
+static const char *databaseName(const char *zName){
+  while( zName[-1]!=0 || zName[-2]!=0 || zName[-3]!=0 || zName[-4]!=0 ){
+    zName--;
+  }
+  return zName;
+}
+
 /*
 ** This is a utility routine, useful to VFS implementations, that checks
 ** to see if a database file was a URI that contained a specific query 
@@ -158738,6 +163328,7 @@ SQLITE_API int sqlite3_test_control(int op, ...){
 */
 SQLITE_API const char *sqlite3_uri_parameter(const char *zFilename, const char *zParam){
   if( zFilename==0 || zParam==0 ) return 0;
+  zFilename = databaseName(zFilename);
   zFilename += sqlite3Strlen30(zFilename) + 1;
   while( zFilename[0] ){
     int x = strcmp(zFilename, zParam);
@@ -158748,6 +163339,20 @@ SQLITE_API const char *sqlite3_uri_parameter(const char *zFilename, const char *
   return 0;
 }
 
+/*
+** Return a pointer to the name of Nth query parameter of the filename.
+*/
+SQLITE_API const char *sqlite3_uri_key(const char *zFilename, int N){
+  if( zFilename==0 || N<0 ) return 0;
+  zFilename = databaseName(zFilename);
+  zFilename += sqlite3Strlen30(zFilename) + 1;
+  while( zFilename[0] && (N--)>0 ){
+    zFilename += sqlite3Strlen30(zFilename) + 1;
+    zFilename += sqlite3Strlen30(zFilename) + 1;
+  }
+  return zFilename[0] ? zFilename : 0;
+}
+
 /*
 ** Return a boolean value for a query parameter.
 */
@@ -158773,6 +163378,39 @@ SQLITE_API sqlite3_int64 sqlite3_uri_int64(
   return bDflt;
 }
 
+/*
+** Translate a filename that was handed to a VFS routine into the corresponding
+** database, journal, or WAL file.
+**
+** It is an error to pass this routine a filename string that was not
+** passed into the VFS from the SQLite core.  Doing so is similar to
+** passing free() a pointer that was not obtained from malloc() - it is
+** an error that we cannot easily detect but that will likely cause memory
+** corruption.
+*/
+SQLITE_API const char *sqlite3_filename_database(const char *zFilename){
+  return databaseName(zFilename);
+  return sqlite3_uri_parameter(zFilename - 3, "\003");
+}
+SQLITE_API const char *sqlite3_filename_journal(const char *zFilename){
+  zFilename = databaseName(zFilename);
+  zFilename += sqlite3Strlen30(zFilename) + 1;
+  while( zFilename[0] ){
+    zFilename += sqlite3Strlen30(zFilename) + 1;
+    zFilename += sqlite3Strlen30(zFilename) + 1;
+  }
+  return zFilename + 1;
+}
+SQLITE_API const char *sqlite3_filename_wal(const char *zFilename){
+#ifdef SQLITE_OMIT_WAL
+  return 0;
+#else
+  zFilename = sqlite3_filename_journal(zFilename);
+  zFilename += sqlite3Strlen30(zFilename) + 1;
+  return zFilename;
+#endif
+}
+
 /*
 ** Return the Btree pointer identified by zDbName.  Return NULL if not found.
 */
@@ -160106,6 +164744,9 @@ typedef sqlite3_int64 i64;        /* 8-byte signed integer */
 # define TESTONLY(X)
 #endif
 
+#define LARGEST_INT64  (0xffffffff|(((i64)0x7fffffff)<<32))
+#define SMALLEST_INT64 (((i64)-1) - LARGEST_INT64)
+
 #endif /* SQLITE_AMALGAMATION */
 
 #ifdef SQLITE_DEBUG
@@ -160149,6 +164790,7 @@ struct Fts3Table {
   char *zLanguageid;              /* languageid=xxx option, or NULL */
   int nAutoincrmerge;             /* Value configured by 'automerge' */
   u32 nLeafAdd;                   /* Number of leaf blocks added this trans */
+  int bLock;                      /* Used to prevent recursive content= tbls */
 
   /* Precompiled statements used by the implementation. Each of these 
   ** statements is run and reset within a single virtual table API call. 
@@ -160207,13 +164849,23 @@ struct Fts3Table {
   int mxSavepoint;       /* Largest valid xSavepoint integer */
 #endif
 
-#ifdef SQLITE_TEST
+#if defined(SQLITE_DEBUG) || defined(SQLITE_TEST)
   /* True to disable the incremental doclist optimization. This is controled
   ** by special insert command 'test-no-incr-doclist'.  */
   int bNoIncrDoclist;
+
+  /* Number of segments in a level */
+  int nMergeCount;
 #endif
 };
 
+/* Macro to find the number of segments to merge */
+#if defined(SQLITE_DEBUG) || defined(SQLITE_TEST)
+# define MergeCount(P) ((P)->nMergeCount)
+#else
+# define MergeCount(P) FTS3_MERGE_COUNT
+#endif
+
 /*
 ** When the core wants to read from the virtual table, it creates a
 ** virtual table cursor (an instance of the following structure) using
@@ -160477,6 +165129,8 @@ SQLITE_PRIVATE int sqlite3Fts3Incrmerge(Fts3Table*,int,int);
 SQLITE_PRIVATE void sqlite3Fts3ErrMsg(char**,const char*,...);
 SQLITE_PRIVATE int sqlite3Fts3PutVarint(char *, sqlite3_int64);
 SQLITE_PRIVATE int sqlite3Fts3GetVarint(const char *, sqlite_int64 *);
+SQLITE_PRIVATE int sqlite3Fts3GetVarintU(const char *, sqlite_uint64 *);
+SQLITE_PRIVATE int sqlite3Fts3GetVarintBounded(const char*,const char*,sqlite3_int64*);
 SQLITE_PRIVATE int sqlite3Fts3GetVarint32(const char *, int *);
 SQLITE_PRIVATE int sqlite3Fts3VarintLen(sqlite3_uint64);
 SQLITE_PRIVATE void sqlite3Fts3Dequote(char *);
@@ -160607,12 +165261,7 @@ SQLITE_PRIVATE int sqlite3Fts3PutVarint(char *p, sqlite_int64 v){
   v = (*ptr++);                                               \
   if( (v & mask2)==0 ){ var = v; return ret; }
 
-/* 
-** Read a 64-bit variable-length integer from memory starting at p[0].
-** Return the number of bytes read, or 0 on error.
-** The value is stored in *v.
-*/
-SQLITE_PRIVATE int sqlite3Fts3GetVarint(const char *pBuf, sqlite_int64 *v){
+SQLITE_PRIVATE int sqlite3Fts3GetVarintU(const char *pBuf, sqlite_uint64 *v){
   const unsigned char *p = (const unsigned char*)pBuf;
   const unsigned char *pStart = p;
   u32 a;
@@ -160634,6 +165283,41 @@ SQLITE_PRIVATE int sqlite3Fts3GetVarint(const char *pBuf, sqlite_int64 *v){
   return (int)(p - pStart);
 }
 
+/* 
+** Read a 64-bit variable-length integer from memory starting at p[0].
+** Return the number of bytes read, or 0 on error.
+** The value is stored in *v.
+*/
+SQLITE_PRIVATE int sqlite3Fts3GetVarint(const char *pBuf, sqlite_int64 *v){
+  return sqlite3Fts3GetVarintU(pBuf, (sqlite3_uint64*)v);
+}
+
+/* 
+** Read a 64-bit variable-length integer from memory starting at p[0] and
+** not extending past pEnd[-1].
+** Return the number of bytes read, or 0 on error.
+** The value is stored in *v.
+*/
+SQLITE_PRIVATE int sqlite3Fts3GetVarintBounded(
+  const char *pBuf,
+  const char *pEnd,
+  sqlite_int64 *v
+){
+  const unsigned char *p = (const unsigned char*)pBuf;
+  const unsigned char *pStart = p;
+  const unsigned char *pX = (const unsigned char*)pEnd;
+  u64 b = 0;
+  int shift;
+  for(shift=0; shift<=63; shift+=7){
+    u64 c = p<pX ? *p : 0;
+    p++;
+    b += (c&0x7F) << shift;
+    if( (c & 0x80)==0 ) break;
+  }
+  *v = b;
+  return (int)(p - pStart);
+}
+
 /*
 ** Similar to sqlite3Fts3GetVarint(), except that the output is truncated to 
 ** a non-negative 32-bit integer before it is returned.
@@ -161729,6 +166413,10 @@ static int fts3InitVtab(
   fts3DatabasePageSize(&rc, p);
   p->nNodeSize = p->nPgsz-35;
 
+#if defined(SQLITE_DEBUG)||defined(SQLITE_TEST)
+  p->nMergeCount = FTS3_MERGE_COUNT;
+#endif
+
   /* Declare the table schema to SQLite. */
   fts3DeclareVtab(&rc, p);
 
@@ -161824,6 +166512,10 @@ static int fts3BestIndexMethod(sqlite3_vtab *pVTab, sqlite3_index_info *pInfo){
   int iDocidLe = -1;              /* Index of docid<=x constraint, if present */
   int iIdx;
 
+  if( p->bLock ){
+    return SQLITE_ERROR;
+  }
+
   /* By default use a full table scan. This is an expensive option,
   ** so search through the constraints to see if a more efficient 
   ** strategy is possible.
@@ -162022,7 +166714,11 @@ static int fts3CursorSeekStmt(Fts3Cursor *pCsr){
     }else{
       zSql = sqlite3_mprintf("SELECT %s WHERE rowid = ?", p->zReadExprlist);
       if( !zSql ) return SQLITE_NOMEM;
-      rc = sqlite3_prepare_v3(p->db, zSql,-1,SQLITE_PREPARE_PERSISTENT,&pCsr->pStmt,0);
+      p->bLock++;
+      rc = sqlite3_prepare_v3(
+          p->db, zSql,-1,SQLITE_PREPARE_PERSISTENT,&pCsr->pStmt,0
+      );
+      p->bLock--;
       sqlite3_free(zSql);
     }
     if( rc==SQLITE_OK ) pCsr->bSeekStmt = 1;
@@ -162040,11 +166736,15 @@ static int fts3CursorSeek(sqlite3_context *pContext, Fts3Cursor *pCsr){
   if( pCsr->isRequireSeek ){
     rc = fts3CursorSeekStmt(pCsr);
     if( rc==SQLITE_OK ){
+      Fts3Table *pTab = (Fts3Table*)pCsr->base.pVtab;
+      pTab->bLock++;
       sqlite3_bind_int64(pCsr->pStmt, 1, pCsr->iPrevId);
       pCsr->isRequireSeek = 0;
       if( SQLITE_ROW==sqlite3_step(pCsr->pStmt) ){
+        pTab->bLock--;
         return SQLITE_OK;
       }else{
+        pTab->bLock--;
         rc = sqlite3_reset(pCsr->pStmt);
         if( rc==SQLITE_OK && ((Fts3Table *)pCsr->base.pVtab)->zContentTbl==0 ){
           /* If no row was found and no error has occurred, then the %_content
@@ -162216,7 +166916,7 @@ static int fts3SelectLeaf(
 
   fts3GetVarint32(zNode, &iHeight);
   rc = fts3ScanInteriorNode(zTerm, nTerm, zNode, nNode, piLeaf, piLeaf2);
-  assert( !piLeaf2 || !piLeaf || rc!=SQLITE_OK || (*piLeaf<=*piLeaf2) );
+  assert_fts3_nc( !piLeaf2 || !piLeaf || rc!=SQLITE_OK || (*piLeaf<=*piLeaf2) );
 
   if( rc==SQLITE_OK && iHeight>1 ){
     char *zBlob = 0;              /* Blob read from %_segments table */
@@ -162236,7 +166936,13 @@ static int fts3SelectLeaf(
       rc = sqlite3Fts3ReadBlock(p, piLeaf?*piLeaf:*piLeaf2, &zBlob, &nBlob, 0);
     }
     if( rc==SQLITE_OK ){
-      rc = fts3SelectLeaf(p, zTerm, nTerm, zBlob, nBlob, piLeaf, piLeaf2);
+      int iNewHeight = 0;
+      fts3GetVarint32(zBlob, &iNewHeight);
+      if( iNewHeight>=iHeight ){
+        rc = FTS_CORRUPT_VTAB;
+      }else{
+        rc = fts3SelectLeaf(p, zTerm, nTerm, zBlob, nBlob, piLeaf, piLeaf2);
+      }
     }
     sqlite3_free(zBlob);
   }
@@ -162341,10 +167047,11 @@ static void fts3ColumnlistCopy(char **pp, char **ppPoslist){
 }
 
 /*
-** Value used to signify the end of an position-list. This is safe because
-** it is not possible to have a document with 2^31 terms.
+** Value used to signify the end of an position-list. This must be
+** as large or larger than any value that might appear on the
+** position-list, even a position list that has been corrupted.
 */
-#define POSITION_LIST_END 0x7fffffff
+#define POSITION_LIST_END LARGEST_INT64
 
 /*
 ** This function is used to help parse position-lists. When this function is
@@ -162420,14 +167127,14 @@ static int fts3PoslistMerge(
       fts3GetVarint32(&p1[1], &iCol1);
       if( iCol1==0 ) return FTS_CORRUPT_VTAB;
     }
-    else if( *p1==POS_END ) iCol1 = POSITION_LIST_END;
+    else if( *p1==POS_END ) iCol1 = 0x7fffffff;
     else iCol1 = 0;
 
     if( *p2==POS_COLUMN ){
       fts3GetVarint32(&p2[1], &iCol2);
       if( iCol2==0 ) return FTS_CORRUPT_VTAB;
     }
-    else if( *p2==POS_END ) iCol2 = POSITION_LIST_END;
+    else if( *p2==POS_END ) iCol2 = 0x7fffffff;
     else iCol2 = 0;
 
     if( iCol1==iCol2 ){
@@ -162690,12 +167397,12 @@ static void fts3GetDeltaVarint3(
   if( *pp>=pEnd ){
     *pp = 0;
   }else{
-    sqlite3_int64 iVal;
-    *pp += sqlite3Fts3GetVarint(*pp, &iVal);
+    u64 iVal;
+    *pp += sqlite3Fts3GetVarintU(*pp, &iVal);
     if( bDescIdx ){
-      *pVal -= iVal;
+      *pVal = (i64)((u64)*pVal - iVal);
     }else{
-      *pVal += iVal;
+      *pVal = (i64)((u64)*pVal + iVal);
     }
   }
 }
@@ -162722,14 +167429,16 @@ static void fts3PutDeltaVarint3(
   int *pbFirst,                   /* IN/OUT: True after first int written */
   sqlite3_int64 iVal              /* Write this value to the list */
 ){
-  sqlite3_int64 iWrite;
+  sqlite3_uint64 iWrite;
   if( bDescIdx==0 || *pbFirst==0 ){
-    iWrite = iVal - *piPrev;
+    assert_fts3_nc( *pbFirst==0 || iVal>=*piPrev );
+    iWrite = (u64)iVal - (u64)*piPrev;
   }else{
-    iWrite = *piPrev - iVal;
+    assert_fts3_nc( *piPrev>=iVal );
+    iWrite = (u64)*piPrev - (u64)iVal;
   }
   assert( *pbFirst || *piPrev==0 );
-  assert( *pbFirst==0 || iWrite>0 );
+  assert_fts3_nc( *pbFirst==0 || iWrite>0 );
   *pp += sqlite3Fts3PutVarint(*pp, iWrite);
   *piPrev = iVal;
   *pbFirst = 1;
@@ -162745,7 +167454,8 @@ static void fts3PutDeltaVarint3(
 ** Using this makes it easier to write code that can merge doclists that are
 ** sorted in either ascending or descending order.
 */
-#define DOCID_CMP(i1, i2) ((bDescDoclist?-1:1) * (i1-i2))
+/* #define DOCID_CMP(i1, i2) ((bDescDoclist?-1:1) * (i64)((u64)i1-i2)) */
+#define DOCID_CMP(i1, i2) ((bDescDoclist?-1:1) * (i1>i2?1:((i1==i2)?0:-1)))
 
 /*
 ** This function does an "OR" merge of two doclists (output contains all
@@ -162835,6 +167545,8 @@ static int fts3DoclistOrMerge(
       fts3PoslistCopy(&p, &p2);
       fts3GetDeltaVarint3(&p2, pEnd2, bDescDoclist, &i2);
     }
+    
+    assert( (p-aOut)<=((p1?(p1-a1):n1)+(p2?(p2-a2):n2)+FTS3_VARINT_MAX-1) );
   }
 
   if( rc!=SQLITE_OK ){
@@ -163157,7 +167869,7 @@ static int fts3SegReaderCursor(
   ** Fts3SegReaderPending might segfault, as the data structures used by 
   ** fts4aux are not completely populated. So it's easiest to filter these
   ** calls out here.  */
-  if( iLevel<0 && p->aIndex ){
+  if( iLevel<0 && p->aIndex && p->iPrevLangid==iLangid ){
     Fts3SegReader *pSeg = 0;
     rc = sqlite3Fts3SegReaderPending(p, iIndex, zTerm, nTerm, isPrefix||isScan, &pSeg);
     if( rc==SQLITE_OK && pSeg ){
@@ -163420,6 +168132,8 @@ static int fts3NextMethod(sqlite3_vtab_cursor *pCursor){
   int rc;
   Fts3Cursor *pCsr = (Fts3Cursor *)pCursor;
   if( pCsr->eSearch==FTS3_DOCID_SEARCH || pCsr->eSearch==FTS3_FULLSCAN_SEARCH ){
+    Fts3Table *pTab = (Fts3Table*)pCursor->pVtab;
+    pTab->bLock++;
     if( SQLITE_ROW!=sqlite3_step(pCsr->pStmt) ){
       pCsr->isEof = 1;
       rc = sqlite3_reset(pCsr->pStmt);
@@ -163427,6 +168141,7 @@ static int fts3NextMethod(sqlite3_vtab_cursor *pCursor){
       pCsr->iPrevId = sqlite3_column_int64(pCsr->pStmt, 0);
       rc = SQLITE_OK;
     }
+    pTab->bLock--;
   }else{
     rc = fts3EvalNext((Fts3Cursor *)pCursor);
   }
@@ -163434,18 +168149,6 @@ static int fts3NextMethod(sqlite3_vtab_cursor *pCursor){
   return rc;
 }
 
-/*
-** The following are copied from sqliteInt.h.
-**
-** Constants for the largest and smallest possible 64-bit signed integers.
-** These macros are designed to work correctly on both 32-bit and 64-bit
-** compilers.
-*/
-#ifndef SQLITE_AMALGAMATION
-# define LARGEST_INT64  (0xffffffff|(((sqlite3_int64)0x7fffffff)<<32))
-# define SMALLEST_INT64 (((sqlite3_int64)-1) - LARGEST_INT64)
-#endif
-
 /*
 ** If the numeric type of argument pVal is "integer", then return it
 ** converted to a 64-bit signed integer. Otherwise, return a copy of
@@ -163499,6 +168202,10 @@ static int fts3FilterMethod(
   UNUSED_PARAMETER(idxStr);
   UNUSED_PARAMETER(nVal);
 
+  if( p->bLock ){
+    return SQLITE_ERROR;
+  }
+
   eSearch = (idxNum & 0x0000FFFF);
   assert( eSearch>=0 && eSearch<=(FTS3_FULLTEXT_SEARCH+p->nColumn) );
   assert( p->pSegments==0 );
@@ -163570,7 +168277,11 @@ static int fts3FilterMethod(
       );
     }
     if( zSql ){
-      rc = sqlite3_prepare_v3(p->db,zSql,-1,SQLITE_PREPARE_PERSISTENT,&pCsr->pStmt,0);
+      p->bLock++;
+      rc = sqlite3_prepare_v3(
+          p->db,zSql,-1,SQLITE_PREPARE_PERSISTENT,&pCsr->pStmt,0
+      );
+      p->bLock--;
       sqlite3_free(zSql);
     }else{
       rc = SQLITE_NOMEM;
@@ -164587,7 +169298,7 @@ static int fts3EvalPhraseStart(Fts3Cursor *pCsr, int bOptOk, Fts3Phrase *p){
   int bIncrOk = (bOptOk 
    && pCsr->bDesc==pTab->bDescIdx 
    && p->nToken<=MAX_INCR_PHRASE_TOKENS && p->nToken>0
-#ifdef SQLITE_TEST
+#if defined(SQLITE_DEBUG) || defined(SQLITE_TEST)
    && pTab->bNoIncrDoclist==0
 #endif
   );
@@ -164729,15 +169440,16 @@ static void fts3EvalDlPhraseNext(
   u8 *pbEof
 ){
   char *pIter;                            /* Used to iterate through aAll */
-  char *pEnd = &pDL->aAll[pDL->nAll];     /* 1 byte past end of aAll */
+  char *pEnd;                             /* 1 byte past end of aAll */
  
   if( pDL->pNextDocid ){
     pIter = pDL->pNextDocid;
+    assert( pDL->aAll!=0 || pIter==0 );
   }else{
     pIter = pDL->aAll;
   }
 
-  if( pIter>=pEnd ){
+  if( pIter==0 || pIter>=(pEnd = pDL->aAll + pDL->nAll) ){
     /* We have already reached the end of this doclist. EOF. */
     *pbEof = 1;
   }else{
@@ -165109,12 +169821,13 @@ static int fts3EvalAverageDocsize(Fts3Cursor *pCsr, int *pnPage){
     rc = sqlite3Fts3SelectDoctotal(p, &pStmt);
     if( rc!=SQLITE_OK ) return rc;
     a = sqlite3_column_blob(pStmt, 0);
-    assert( a );
-
-    pEnd = &a[sqlite3_column_bytes(pStmt, 0)];
-    a += sqlite3Fts3GetVarint(a, &nDoc);
-    while( a<pEnd ){
-      a += sqlite3Fts3GetVarint(a, &nByte);
+    testcase( a==0 );  /* If %_stat.value set to X'' */
+    if( a ){
+      pEnd = &a[sqlite3_column_bytes(pStmt, 0)];
+      a += sqlite3Fts3GetVarintBounded(a, pEnd, &nDoc);
+      while( a<pEnd ){
+        a += sqlite3Fts3GetVarintBounded(a, pEnd, &nByte);
+      }
     }
     if( nDoc==0 || nByte==0 ){
       sqlite3_reset(pStmt);
@@ -169552,7 +174265,9 @@ int queryTokenizer(
 
   sqlite3_bind_text(pStmt, 1, zName, -1, SQLITE_STATIC);
   if( SQLITE_ROW==sqlite3_step(pStmt) ){
-    if( sqlite3_column_type(pStmt, 0)==SQLITE_BLOB ){
+    if( sqlite3_column_type(pStmt, 0)==SQLITE_BLOB
+     && sqlite3_column_bytes(pStmt, 0)==sizeof(*pp)
+    ){
       memcpy((void *)pp, sqlite3_column_blob(pStmt, 0), sizeof(*pp));
     }
   }
@@ -169641,7 +174356,7 @@ SQLITE_PRIVATE int sqlite3Fts3InitHashTable(
 ){
   int rc = SQLITE_OK;
   void *p = (void *)pHash;
-  const int any = SQLITE_ANY;
+  const int any = SQLITE_UTF8|SQLITE_DIRECTONLY;
 
 #ifdef SQLITE_TEST
   char *zTest = 0;
@@ -170401,7 +175116,7 @@ SQLITE_PRIVATE int sqlite3Fts3InitTok(sqlite3 *db, Fts3Hash *pHash){
 /* #include <string.h> */
 /* #include <assert.h> */
 /* #include <stdlib.h> */
-
+/* #include <stdio.h> */
 
 #define FTS_MAX_APPENDABLE_HEIGHT 16
 
@@ -170445,7 +175160,7 @@ int test_fts3_node_chunk_threshold = (4*1024)*4;
 #endif
 
 /*
-** The two values that may be meaningfully bound to the :1 parameter in
+** The values that may be meaningfully bound to the :1 parameter in
 ** statements SQL_REPLACE_STAT and SQL_SELECT_STAT.
 */
 #define FTS_STAT_DOCTOTAL      0
@@ -170713,7 +175428,7 @@ static int fts3SqlStmt(
 ** returns zero rows.  */
 /* 28 */ "SELECT level, count(*) AS cnt FROM %Q.'%q_segdir' "
          "  GROUP BY level HAVING cnt>=?"
-         "  ORDER BY (level %% 1024) ASC LIMIT 1",
+         "  ORDER BY (level %% 1024) ASC, 2 DESC LIMIT 1",
 
 /* Estimate the upper limit on the number of leaf nodes in a new segment
 ** created by merging the oldest :2 segments from absolute level :1. See 
@@ -171074,7 +175789,7 @@ static int fts3PendingListAppend(
   assert( !p || p->iLastDocid<=iDocid );
 
   if( !p || p->iLastDocid!=iDocid ){
-    sqlite3_int64 iDelta = iDocid - (p ? p->iLastDocid : 0);
+    u64 iDelta = (u64)iDocid - (u64)(p ? p->iLastDocid : 0);
     if( p ){
       assert( p->nData<p->nSpace );
       assert( p->aData[p->nData]==0 );
@@ -171531,7 +176246,7 @@ static int fts3AllocateSegdirIdx(
     ** segment and allocate (newly freed) index 0 at level iLevel. Otherwise,
     ** if iNext is less than FTS3_MERGE_COUNT, allocate index iNext.
     */
-    if( iNext>=FTS3_MERGE_COUNT ){
+    if( iNext>=MergeCount(p) ){
       fts3LogMerge(16, getAbsoluteLevel(p, iLangid, iIndex, iLevel));
       rc = fts3SegmentMerge(p, iLangid, iIndex, iLevel);
       *piIdx = 0;
@@ -171615,6 +176330,8 @@ SQLITE_PRIVATE int sqlite3Fts3ReadBlock(
       }
       *paBlob = aByte;
     }
+  }else if( rc==SQLITE_ERROR ){
+    rc = FTS_CORRUPT_VTAB;
   }
 
   return rc;
@@ -171757,7 +176474,7 @@ static int fts3SegReaderNext(
   pNext += fts3GetVarint32(pNext, &nSuffix);
   if( nSuffix<=0 
    || (&pReader->aNode[pReader->nNode] - pNext)<nSuffix
-   || nPrefix>pReader->nTermAlloc
+   || nPrefix>pReader->nTerm
   ){
     return FTS_CORRUPT_VTAB;
   }
@@ -171907,18 +176624,18 @@ static int fts3SegReaderNextDocid(
     }else{
       rc = fts3SegReaderRequire(pReader, p, FTS3_VARINT_MAX);
       if( rc==SQLITE_OK ){
-        sqlite3_int64 iDelta;
-        pReader->pOffsetList = p + sqlite3Fts3GetVarint(p, &iDelta);
+        u64 iDelta;
+        pReader->pOffsetList = p + sqlite3Fts3GetVarintU(p, &iDelta);
         if( pTab->bDescIdx ){
-          pReader->iDocid -= iDelta;
+          pReader->iDocid = (i64)((u64)pReader->iDocid - iDelta);
         }else{
-          pReader->iDocid += iDelta;
+          pReader->iDocid = (i64)((u64)pReader->iDocid + iDelta);
         }
       }
     }
   }
 
-  return SQLITE_OK;
+  return rc;
 }
 
 
@@ -172408,6 +177125,11 @@ static int fts3NodeAddTerm(
     nPrefix = fts3PrefixCompress(pTree->zTerm, pTree->nTerm, zTerm, nTerm);
     nSuffix = nTerm-nPrefix;
 
+    /* If nSuffix is zero or less, then zTerm/nTerm must be a prefix of 
+    ** pWriter->zTerm/pWriter->nTerm. i.e. must be equal to or less than when
+    ** compared with BINARY collation. This indicates corruption.  */
+    if( nSuffix<=0 ) return FTS_CORRUPT_VTAB;
+
     nReq += sqlite3Fts3VarintLen(nPrefix)+sqlite3Fts3VarintLen(nSuffix)+nSuffix;
     if( nReq<=p->nNodeSize || !pTree->zTerm ){
 
@@ -172652,6 +177374,7 @@ static int fts3SegWriterAdd(
     int rc;
 
     /* The current leaf node is full. Write it out to the database. */
+    if( pWriter->iFree==LARGEST_INT64 ) return FTS_CORRUPT_VTAB;
     rc = fts3WriteSegment(p, pWriter->iFree++, pWriter->aData, nData);
     if( rc!=SQLITE_OK ) return rc;
     p->nLeafAdd++;
@@ -172701,9 +177424,11 @@ static int fts3SegWriterAdd(
   /* Append the prefix-compressed term and doclist to the buffer. */
   nData += sqlite3Fts3PutVarint(&pWriter->aData[nData], nPrefix);
   nData += sqlite3Fts3PutVarint(&pWriter->aData[nData], nSuffix);
+  assert( nSuffix>0 );
   memcpy(&pWriter->aData[nData], &zTerm[nPrefix], nSuffix);
   nData += nSuffix;
   nData += sqlite3Fts3PutVarint(&pWriter->aData[nData], nDoclist);
+  assert( nDoclist>0 );
   memcpy(&pWriter->aData[nData], aDoclist, nDoclist);
   pWriter->nData = nData + nDoclist;
 
@@ -172723,6 +177448,7 @@ static int fts3SegWriterAdd(
       pWriter->zTerm = zNew;
     }
     assert( pWriter->zTerm==pWriter->zMalloc );
+    assert( nTerm>0 );
     memcpy(pWriter->zTerm, zTerm, nTerm);
   }else{
     pWriter->zTerm = (char *)zTerm;
@@ -172997,14 +177723,14 @@ static void fts3ColumnFilter(
 
     nList -= (int)(p - pList);
     pList = p;
-    if( nList==0 ){
+    if( nList<=0 ){
       break;
     }
     p = &pList[1];
     p += fts3GetVarint32(p, &iCurrent);
   }
 
-  if( bZero && &pList[nList]!=pEnd ){
+  if( bZero && (pEnd - &pList[nList])>0){
     memset(&pList[nList], 0, pEnd - &pList[nList]);
   }
   *ppList = pList;
@@ -173031,6 +177757,7 @@ static int fts3MsrBufferData(
     pMsr->aBuffer = pNew;
   }
 
+  assert( nList>0 );
   memcpy(pMsr->aBuffer, pList, nList);
   return SQLITE_OK;
 }
@@ -173344,14 +178071,12 @@ SQLITE_PRIVATE int sqlite3Fts3SegReaderStep(
           ** doclist. */
           sqlite3_int64 iDelta;
           if( p->bDescIdx && nDoclist>0 ){
-            iDelta = iPrev - iDocid;
+            if( iPrev<=iDocid ) return FTS_CORRUPT_VTAB;
+            iDelta = (i64)((u64)iPrev - (u64)iDocid);
           }else{
-            iDelta = iDocid - iPrev;
-          }
-          if( iDelta<=0 && (nDoclist>0 || iDelta!=iDocid) ){
-            return FTS_CORRUPT_VTAB;
+            if( nDoclist>0 && iPrev>=iDocid ) return FTS_CORRUPT_VTAB;
+            iDelta = (i64)((u64)iDocid - (u64)iPrev);
           }
-          assert( nDoclist>0 || iDelta==iDocid );
 
           nByte = sqlite3Fts3VarintLen(iDelta) + (isRequirePos?nList+1:0);
           if( nDoclist+nByte>pCsr->nBuffer ){
@@ -173633,7 +178358,7 @@ static int fts3SegmentMerge(
         csr.zTerm, csr.nTerm, csr.aDoclist, csr.nDoclist);
   }
   if( rc!=SQLITE_OK ) goto finished;
-  assert( pWriter || bIgnoreEmpty );
+  assert_fts3_nc( pWriter || bIgnoreEmpty );
 
   if( iLevel!=FTS3_SEGCURSOR_PENDING ){
     rc = fts3DeleteSegdir(
@@ -173860,7 +178585,10 @@ static int fts3DoOptimize(Fts3Table *p, int bReturnDone){
   int rc;
   sqlite3_stmt *pAllLangid = 0;
 
-  rc = fts3SqlStmt(p, SQL_SELECT_ALL_LANGID, &pAllLangid, 0);
+  rc = sqlite3Fts3PendingTermsFlush(p);
+  if( rc==SQLITE_OK ){
+    rc = fts3SqlStmt(p, SQL_SELECT_ALL_LANGID, &pAllLangid, 0);
+  }
   if( rc==SQLITE_OK ){
     int rc2;
     sqlite3_bind_int(pAllLangid, 1, p->iPrevLangid);
@@ -173881,7 +178609,6 @@ static int fts3DoOptimize(Fts3Table *p, int bReturnDone){
   }
 
   sqlite3Fts3SegmentsClose(p);
-  sqlite3Fts3PendingTermsClear(p);
 
   return (rc==SQLITE_OK && bReturnDone && bSeenDone) ? SQLITE_DONE : rc;
 }
@@ -174132,7 +178859,7 @@ static int nodeReaderNext(NodeReader *p){
     }
     p->iOff += fts3GetVarint32(&p->aNode[p->iOff], &nSuffix);
 
-    if( nPrefix>p->iOff || nSuffix>p->nNode-p->iOff ){
+    if( nPrefix>p->term.n || nSuffix>p->nNode-p->iOff || nSuffix==0 ){
       return FTS_CORRUPT_VTAB;
     }
     blobGrowBuffer(&p->term, nPrefix+nSuffix, &rc);
@@ -174151,7 +178878,7 @@ static int nodeReaderNext(NodeReader *p){
     }
   }
 
-  assert( p->iOff<=p->nNode );
+  assert_fts3_nc( p->iOff<=p->nNode );
   return rc;
 }
 
@@ -174175,14 +178902,14 @@ static int nodeReaderInit(NodeReader *p, const char *aNode, int nNode){
   p->nNode = nNode;
 
   /* Figure out if this is a leaf or an internal node. */
-  if( p->aNode[0] ){
+  if( aNode && aNode[0] ){
     /* An internal node. */
     p->iOff = 1 + sqlite3Fts3GetVarint(&p->aNode[1], &p->iChild);
   }else{
     p->iOff = 1;
   }
 
-  return nodeReaderNext(p);
+  return aNode ? nodeReaderNext(p) : SQLITE_OK;
 }
 
 /*
@@ -174219,6 +178946,7 @@ static int fts3IncrmergePush(
     ** be added to.  */
     nPrefix = fts3PrefixCompress(pNode->key.a, pNode->key.n, zTerm, nTerm);
     nSuffix = nTerm - nPrefix;
+    if(nSuffix<=0 ) return FTS_CORRUPT_VTAB;
     nSpace  = sqlite3Fts3VarintLen(nPrefix);
     nSpace += sqlite3Fts3VarintLen(nSuffix) + nSuffix;
 
@@ -174312,13 +179040,14 @@ static int fts3AppendToNode(
   /* Node must have already been started. There must be a doclist for a
   ** leaf node, and there must not be a doclist for an internal node.  */
   assert( pNode->n>0 );
-  assert( (pNode->a[0]=='\0')==(aDoclist!=0) );
+  assert_fts3_nc( (pNode->a[0]=='\0')==(aDoclist!=0) );
 
   blobGrowBuffer(pPrev, nTerm, &rc);
   if( rc!=SQLITE_OK ) return rc;
 
   nPrefix = fts3PrefixCompress(pPrev->a, pPrev->n, zTerm, nTerm);
   nSuffix = nTerm - nPrefix;
+  if( nSuffix<=0 ) return FTS_CORRUPT_VTAB;
   memcpy(pPrev->a, zTerm, nTerm);
   pPrev->n = nTerm;
 
@@ -174528,7 +179257,7 @@ static int fts3TermCmp(
   int nCmp = MIN(nLhs, nRhs);
   int res;
 
-  res = memcmp(zLhs, zRhs, nCmp);
+  res = (nCmp ? memcmp(zLhs, zRhs, nCmp) : 0);
   if( res==0 ) res = nLhs - nRhs;
 
   return res;
@@ -174612,6 +179341,10 @@ static int fts3IncrmergeLoad(
       pWriter->bNoLeafData = (pWriter->nLeafData==0);
       nRoot = sqlite3_column_bytes(pSelect, 4);
       aRoot = sqlite3_column_blob(pSelect, 4);
+      if( aRoot==0 ){
+        sqlite3_reset(pSelect);
+        return nRoot ? SQLITE_NOMEM : FTS_CORRUPT_VTAB;
+      }
     }else{
       return sqlite3_reset(pSelect);
     }
@@ -174647,6 +179380,10 @@ static int fts3IncrmergeLoad(
       int i;
       int nHeight = (int)aRoot[0];
       NodeWriter *pNode;
+      if( nHeight<1 || nHeight>FTS_MAX_APPENDABLE_HEIGHT ){
+        sqlite3_reset(pSelect);
+        return FTS_CORRUPT_VTAB;
+      }
 
       pWriter->nLeafEst = (int)((iEnd - iStart) + 1)/FTS_MAX_APPENDABLE_HEIGHT;
       pWriter->iStart = iStart;
@@ -174660,34 +179397,42 @@ static int fts3IncrmergeLoad(
 
       pNode = &pWriter->aNodeWriter[nHeight];
       pNode->iBlock = pWriter->iStart + pWriter->nLeafEst*nHeight;
-      blobGrowBuffer(&pNode->block, MAX(nRoot, p->nNodeSize), &rc);
+      blobGrowBuffer(&pNode->block, 
+          MAX(nRoot, p->nNodeSize)+FTS3_NODE_PADDING, &rc
+      );
       if( rc==SQLITE_OK ){
         memcpy(pNode->block.a, aRoot, nRoot);
         pNode->block.n = nRoot;
+        memset(&pNode->block.a[nRoot], 0, FTS3_NODE_PADDING);
       }
 
       for(i=nHeight; i>=0 && rc==SQLITE_OK; i--){
         NodeReader reader;
         pNode = &pWriter->aNodeWriter[i];
 
-        rc = nodeReaderInit(&reader, pNode->block.a, pNode->block.n);
-        while( reader.aNode && rc==SQLITE_OK ) rc = nodeReaderNext(&reader);
-        blobGrowBuffer(&pNode->key, reader.term.n, &rc);
-        if( rc==SQLITE_OK ){
-          memcpy(pNode->key.a, reader.term.a, reader.term.n);
-          pNode->key.n = reader.term.n;
-          if( i>0 ){
-            char *aBlock = 0;
-            int nBlock = 0;
-            pNode = &pWriter->aNodeWriter[i-1];
-            pNode->iBlock = reader.iChild;
-            rc = sqlite3Fts3ReadBlock(p, reader.iChild, &aBlock, &nBlock, 0);
-            blobGrowBuffer(&pNode->block, MAX(nBlock, p->nNodeSize), &rc);
-            if( rc==SQLITE_OK ){
-              memcpy(pNode->block.a, aBlock, nBlock);
-              pNode->block.n = nBlock;
+        if( pNode->block.a){
+          rc = nodeReaderInit(&reader, pNode->block.a, pNode->block.n);
+          while( reader.aNode && rc==SQLITE_OK ) rc = nodeReaderNext(&reader);
+          blobGrowBuffer(&pNode->key, reader.term.n, &rc);
+          if( rc==SQLITE_OK ){
+            memcpy(pNode->key.a, reader.term.a, reader.term.n);
+            pNode->key.n = reader.term.n;
+            if( i>0 ){
+              char *aBlock = 0;
+              int nBlock = 0;
+              pNode = &pWriter->aNodeWriter[i-1];
+              pNode->iBlock = reader.iChild;
+              rc = sqlite3Fts3ReadBlock(p, reader.iChild, &aBlock, &nBlock, 0);
+              blobGrowBuffer(&pNode->block, 
+                  MAX(nBlock, p->nNodeSize)+FTS3_NODE_PADDING, &rc
+              );
+              if( rc==SQLITE_OK ){
+                memcpy(pNode->block.a, aBlock, nBlock);
+                pNode->block.n = nBlock;
+                memset(&pNode->block.a[nBlock], 0, FTS3_NODE_PADDING);
+              }
+              sqlite3_free(aBlock);
             }
-            sqlite3_free(aBlock);
           }
         }
         nodeReaderRelease(&reader);
@@ -174930,7 +179675,10 @@ static int fts3TruncateNode(
   NodeReader reader;              /* Reader object */
   Blob prev = {0, 0, 0};          /* Previous term written to new node */
   int rc = SQLITE_OK;             /* Return code */
-  int bLeaf = aNode[0]=='\0';     /* True for a leaf node */
+  int bLeaf;                       /* True for a leaf node */
+
+  if( nNode<1 ) return FTS_CORRUPT_VTAB;
+  bLeaf = aNode[0]=='\0';
 
   /* Allocate required output space */
   blobGrowBuffer(pNew, nNode, &rc);
@@ -175196,13 +179944,17 @@ static int fts3IncrmergeHintPop(Blob *pHint, i64 *piAbsLevel, int *pnInput){
   const int nHint = pHint->n;
   int i;
 
-  i = pHint->n-2;
+  i = pHint->n-1;
+  if( (pHint->a[i] & 0x80) ) return FTS_CORRUPT_VTAB;
   while( i>0 && (pHint->a[i-1] & 0x80) ) i--;
+  if( i==0 ) return FTS_CORRUPT_VTAB;
+  i--;
   while( i>0 && (pHint->a[i-1] & 0x80) ) i--;
 
   pHint->n = i;
   i += sqlite3Fts3GetVarint(&pHint->a[i], piAbsLevel);
   i += fts3GetVarint32(&pHint->a[i], pnInput);
+  assert( i<=nHint );
   if( i!=nHint ) return FTS_CORRUPT_VTAB;
 
   return SQLITE_OK;
@@ -175238,7 +179990,7 @@ SQLITE_PRIVATE int sqlite3Fts3Incrmerge(Fts3Table *p, int nMerge, int nMin){
 
   rc = fts3IncrmergeHintLoad(p, &hint);
   while( rc==SQLITE_OK && nRem>0 ){
-    const i64 nMod = (i64)FTS3_SEGDIR_MAXLEVEL * p->nIndex;
+    const i64 nMod = FTS3_SEGDIR_MAXLEVEL * p->nIndex;
     sqlite3_stmt *pFindLevel = 0; /* SQL used to determine iAbsLevel */
     int bUseHint = 0;             /* True if attempting to append */
     int iIdx = 0;                 /* Largest idx in level (iAbsLevel+1) */
@@ -175272,8 +180024,14 @@ SQLITE_PRIVATE int sqlite3Fts3Incrmerge(Fts3Table *p, int nMerge, int nMin){
 
       rc = fts3IncrmergeHintPop(&hint, &iHintAbsLevel, &nHintSeg);
       if( nSeg<0 || (iAbsLevel % nMod) >= (iHintAbsLevel % nMod) ){
+        /* Based on the scan in the block above, it is known that there
+        ** are no levels with a relative level smaller than that of
+        ** iAbsLevel with more than nSeg segments, or if nSeg is -1, 
+        ** no levels with more than nMin segments. Use this to limit the
+        ** value of nHintSeg to avoid a large memory allocation in case the 
+        ** merge-hint is corrupt*/
         iAbsLevel = iHintAbsLevel;
-        nSeg = nHintSeg;
+        nSeg = MIN(MAX(nMin,nSeg), nHintSeg);
         bUseHint = 1;
         bDirtyHint = 1;
       }else{
@@ -175286,7 +180044,7 @@ SQLITE_PRIVATE int sqlite3Fts3Incrmerge(Fts3Table *p, int nMerge, int nMin){
     /* If nSeg is less that zero, then there is no level with at least
     ** nMin segments and no hint in the %_stat table. No work to do.
     ** Exit early in this case.  */
-    if( nSeg<0 ) break;
+    if( nSeg<=0 ) break;
 
     /* Open a cursor to iterate through the contents of the oldest nSeg 
     ** indexes of absolute level iAbsLevel. If this cursor is opened using 
@@ -175314,8 +180072,15 @@ SQLITE_PRIVATE int sqlite3Fts3Incrmerge(Fts3Table *p, int nMerge, int nMin){
     }
     if( SQLITE_OK==rc && pCsr->nSegment==nSeg
      && SQLITE_OK==(rc = sqlite3Fts3SegReaderStart(p, pCsr, pFilter))
-     && SQLITE_ROW==(rc = sqlite3Fts3SegReaderStep(p, pCsr))
     ){
+      int bEmpty = 0;
+      rc = sqlite3Fts3SegReaderStep(p, pCsr);
+      if( rc==SQLITE_OK ){
+        bEmpty = 1;
+      }else if( rc!=SQLITE_ROW ){
+        sqlite3Fts3SegReaderFinish(pCsr);
+        break;
+      }
       if( bUseHint && iIdx>0 ){
         const char *zKey = pCsr->zTerm;
         int nKey = pCsr->nTerm;
@@ -175326,11 +180091,13 @@ SQLITE_PRIVATE int sqlite3Fts3Incrmerge(Fts3Table *p, int nMerge, int nMin){
 
       if( rc==SQLITE_OK && pWriter->nLeafEst ){
         fts3LogMerge(nSeg, iAbsLevel);
-        do {
-          rc = fts3IncrmergeAppend(p, pWriter, pCsr);
-          if( rc==SQLITE_OK ) rc = sqlite3Fts3SegReaderStep(p, pCsr);
-          if( pWriter->nWork>=nRem && rc==SQLITE_ROW ) rc = SQLITE_OK;
-        }while( rc==SQLITE_ROW );
+        if( bEmpty==0 ){
+          do {
+            rc = fts3IncrmergeAppend(p, pWriter, pCsr);
+            if( rc==SQLITE_OK ) rc = sqlite3Fts3SegReaderStep(p, pCsr);
+            if( pWriter->nWork>=nRem && rc==SQLITE_ROW ) rc = SQLITE_OK;
+          }while( rc==SQLITE_ROW );
+        }
 
         /* Update or delete the input segments */
         if( rc==SQLITE_OK ){
@@ -175395,7 +180162,7 @@ static int fts3DoIncrmerge(
   const char *zParam              /* Nul-terminated string containing "A,B" */
 ){
   int rc;
-  int nMin = (FTS3_MERGE_COUNT / 2);
+  int nMin = (MergeCount(p) / 2);
   int nMerge = 0;
   const char *z = zParam;
 
@@ -175440,7 +180207,7 @@ static int fts3DoAutoincrmerge(
   int rc = SQLITE_OK;
   sqlite3_stmt *pStmt = 0;
   p->nAutoincrmerge = fts3Getint(&zParam);
-  if( p->nAutoincrmerge==1 || p->nAutoincrmerge>FTS3_MERGE_COUNT ){
+  if( p->nAutoincrmerge==1 || p->nAutoincrmerge>MergeCount(p) ){
     p->nAutoincrmerge = 8;
   }
   if( !p->bHasStat ){
@@ -175523,12 +180290,12 @@ static u64 fts3ChecksumIndex(
 
       i64 iDocid = 0;
       i64 iCol = 0;
-      i64 iPos = 0;
+      u64 iPos = 0;
 
       pCsr += sqlite3Fts3GetVarint(pCsr, &iDocid);
       while( pCsr<pEnd ){
-        i64 iVal = 0;
-        pCsr += sqlite3Fts3GetVarint(pCsr, &iVal);
+        u64 iVal = 0;
+        pCsr += sqlite3Fts3GetVarintU(pCsr, &iVal);
         if( pCsr<pEnd ){
           if( iVal==0 || iVal==1 ){
             iCol = 0;
@@ -175536,8 +180303,12 @@ static u64 fts3ChecksumIndex(
             if( iVal ){
               pCsr += sqlite3Fts3GetVarint(pCsr, &iCol);
             }else{
-              pCsr += sqlite3Fts3GetVarint(pCsr, &iVal);
-              iDocid += iVal;
+              pCsr += sqlite3Fts3GetVarintU(pCsr, &iVal);
+              if( p->bDescIdx ){
+                iDocid = (i64)((u64)iDocid - iVal);
+              }else{
+                iDocid = (i64)((u64)iDocid + iVal);
+              }
             }
           }else{
             iPos += (iVal - 2);
@@ -175610,10 +180381,9 @@ static int fts3IntegrityCheck(Fts3Table *p, int *pbOk){
       for(iCol=0; rc==SQLITE_OK && iCol<p->nColumn; iCol++){
         if( p->abNotindexed[iCol]==0 ){
           const char *zText = (const char *)sqlite3_column_text(pStmt, iCol+1);
-          int nText = sqlite3_column_bytes(pStmt, iCol+1);
           sqlite3_tokenizer_cursor *pT = 0;
 
-          rc = sqlite3Fts3OpenTokenizer(p->pTokenizer, iLang, zText, nText,&pT);
+          rc = sqlite3Fts3OpenTokenizer(p->pTokenizer, iLang, zText, -1, &pT);
           while( rc==SQLITE_OK ){
             char const *zToken;       /* Buffer containing token */
             int nToken = 0;           /* Number of bytes in token */
@@ -175698,7 +180468,7 @@ static int fts3DoIntegrityCheck(
 ** meaningful value to insert is the text 'optimize'.
 */
 static int fts3SpecialInsert(Fts3Table *p, sqlite3_value *pVal){
-  int rc;                         /* Return Code */
+  int rc = SQLITE_ERROR;           /* Return Code */
   const char *zVal = (const char *)sqlite3_value_text(pVal);
   int nVal = sqlite3_value_bytes(pVal);
 
@@ -175714,21 +180484,27 @@ static int fts3SpecialInsert(Fts3Table *p, sqlite3_value *pVal){
     rc = fts3DoIncrmerge(p, &zVal[6]);
   }else if( nVal>10 && 0==sqlite3_strnicmp(zVal, "automerge=", 10) ){
     rc = fts3DoAutoincrmerge(p, &zVal[10]);
-#ifdef SQLITE_TEST
-  }else if( nVal>9 && 0==sqlite3_strnicmp(zVal, "nodesize=", 9) ){
-    p->nNodeSize = atoi(&zVal[9]);
-    rc = SQLITE_OK;
-  }else if( nVal>11 && 0==sqlite3_strnicmp(zVal, "maxpending=", 9) ){
-    p->nMaxPendingData = atoi(&zVal[11]);
-    rc = SQLITE_OK;
-  }else if( nVal>21 && 0==sqlite3_strnicmp(zVal, "test-no-incr-doclist=", 21) ){
-    p->bNoIncrDoclist = atoi(&zVal[21]);
-    rc = SQLITE_OK;
-#endif
+#if defined(SQLITE_DEBUG) || defined(SQLITE_TEST)
   }else{
-    rc = SQLITE_ERROR;
+    int v;
+    if( nVal>9 && 0==sqlite3_strnicmp(zVal, "nodesize=", 9) ){
+      v = atoi(&zVal[9]);
+      if( v>=24 && v<=p->nPgsz-35 ) p->nNodeSize = v;
+      rc = SQLITE_OK;
+    }else if( nVal>11 && 0==sqlite3_strnicmp(zVal, "maxpending=", 9) ){
+      v = atoi(&zVal[11]);
+      if( v>=64 && v<=FTS3_MAX_PENDING_DATA ) p->nMaxPendingData = v;
+      rc = SQLITE_OK;
+    }else if( nVal>21 && 0==sqlite3_strnicmp(zVal,"test-no-incr-doclist=",21) ){
+      p->bNoIncrDoclist = atoi(&zVal[21]);
+      rc = SQLITE_OK;
+    }else if( nVal>11 && 0==sqlite3_strnicmp(zVal,"mergecount=",11) ){
+      v = atoi(&zVal[11]);
+      if( v>=4 && v<=FTS3_MERGE_COUNT && (v&1)==0 ) p->nMergeCount = v;
+      rc = SQLITE_OK;
+    }
+#endif
   }
-
   return rc;
 }
 
@@ -176528,10 +181304,10 @@ static void fts3SnippetDetails(
 
       while( iCsr<(iStart+pIter->nSnippet) && iCsr>=iStart ){
         int j;
-        u64 mPhrase = (u64)1 << i;
+        u64 mPhrase = (u64)1 << (i%64);
         u64 mPos = (u64)1 << (iCsr - iStart);
         assert( iCsr>=iStart && (iCsr - iStart)<=64 );
-        assert( i>=0 && i<=64 );
+        assert( i>=0 );
         if( (mCover|mCovered)&mPhrase ){
           iScore++;
         }else{
@@ -176655,7 +181431,7 @@ static int fts3BestSnippet(
     /* Set the *pmSeen output variable. */
     for(i=0; i<nList; i++){
       if( sIter.aPhrase[i].pHead ){
-        *pmSeen |= (u64)1 << i;
+        *pmSeen |= (u64)1 << (i%64);
       }
     }
 
@@ -177133,11 +181909,15 @@ static int fts3MatchinfoSelectDoctotal(
   Fts3Table *pTab,
   sqlite3_stmt **ppStmt,
   sqlite3_int64 *pnDoc,
-  const char **paLen
+  const char **paLen,
+  const char **ppEnd
 ){
   sqlite3_stmt *pStmt;
   const char *a;
+  const char *pEnd;
   sqlite3_int64 nDoc;
+  int n;
+
 
   if( !*ppStmt ){
     int rc = sqlite3Fts3SelectDoctotal(pTab, ppStmt);
@@ -177146,12 +181926,20 @@ static int fts3MatchinfoSelectDoctotal(
   pStmt = *ppStmt;
   assert( sqlite3_data_count(pStmt)==1 );
 
+  n = sqlite3_column_bytes(pStmt, 0);
   a = sqlite3_column_blob(pStmt, 0);
-  a += sqlite3Fts3GetVarint(a, &nDoc);
-  if( nDoc==0 ) return FTS_CORRUPT_VTAB;
-  *pnDoc = (u32)nDoc;
+  if( a==0 ){
+    return FTS_CORRUPT_VTAB;
+  }
+  pEnd = a + n;
+  a += sqlite3Fts3GetVarintBounded(a, pEnd, &nDoc);
+  if( nDoc<=0 || a>pEnd ){
+    return FTS_CORRUPT_VTAB;
+  }
+  *pnDoc = nDoc;
 
   if( paLen ) *paLen = a;
+  if( ppEnd ) *ppEnd = pEnd;
   return SQLITE_OK;
 }
 
@@ -177332,7 +182120,7 @@ static int fts3MatchinfoValues(
       case FTS3_MATCHINFO_NDOC:
         if( bGlobal ){
           sqlite3_int64 nDoc = 0;
-          rc = fts3MatchinfoSelectDoctotal(pTab, &pSelect, &nDoc, 0);
+          rc = fts3MatchinfoSelectDoctotal(pTab, &pSelect, &nDoc, 0, 0);
           pInfo->aMatchinfo[0] = (u32)nDoc;
         }
         break;
@@ -177341,14 +182129,19 @@ static int fts3MatchinfoValues(
         if( bGlobal ){
           sqlite3_int64 nDoc;     /* Number of rows in table */
           const char *a;          /* Aggregate column length array */
+          const char *pEnd;       /* First byte past end of length array */
 
-          rc = fts3MatchinfoSelectDoctotal(pTab, &pSelect, &nDoc, &a);
+          rc = fts3MatchinfoSelectDoctotal(pTab, &pSelect, &nDoc, &a, &pEnd);
           if( rc==SQLITE_OK ){
             int iCol;
             for(iCol=0; iCol<pInfo->nCol; iCol++){
               u32 iVal;
               sqlite3_int64 nToken;
               a += sqlite3Fts3GetVarint(a, &nToken);
+              if( a>pEnd ){
+                rc = SQLITE_CORRUPT_VTAB;
+                break;
+              }
               iVal = (u32)(((u32)(nToken&0xffffffff)+nDoc/2)/nDoc);
               pInfo->aMatchinfo[iCol] = iVal;
             }
@@ -177362,9 +182155,14 @@ static int fts3MatchinfoValues(
         if( rc==SQLITE_OK ){
           int iCol;
           const char *a = sqlite3_column_blob(pSelectDocsize, 0);
+          const char *pEnd = a + sqlite3_column_bytes(pSelectDocsize, 0);
           for(iCol=0; iCol<pInfo->nCol; iCol++){
             sqlite3_int64 nToken;
-            a += sqlite3Fts3GetVarint(a, &nToken);
+            a += sqlite3Fts3GetVarintBounded(a, pEnd, &nToken);
+            if( a>pEnd ){
+              rc = SQLITE_CORRUPT_VTAB;
+              break;
+            }
             pInfo->aMatchinfo[iCol] = (u32)nToken;
           }
         }
@@ -177395,7 +182193,7 @@ static int fts3MatchinfoValues(
         if( rc!=SQLITE_OK ) break;
         if( bGlobal ){
           if( pCsr->pDeferred ){
-            rc = fts3MatchinfoSelectDoctotal(pTab, &pSelect, &pInfo->nDoc, 0);
+            rc = fts3MatchinfoSelectDoctotal(pTab, &pSelect, &pInfo->nDoc,0,0);
             if( rc!=SQLITE_OK ) break;
           }
           rc = fts3ExprIterate(pExpr, fts3ExprGlobalHitsCb,(void*)pInfo);
@@ -179131,6 +183929,37 @@ static void jsonReturnJson(
   sqlite3_result_subtype(pCtx, JSON_SUBTYPE);
 }
 
+/*
+** Translate a single byte of Hex into an integer.
+** This routine only works if h really is a valid hexadecimal
+** character:  0..9a..fA..F
+*/
+static u8 jsonHexToInt(int h){
+  assert( (h>='0' && h<='9') ||  (h>='a' && h<='f') ||  (h>='A' && h<='F') );
+#ifdef SQLITE_EBCDIC
+  h += 9*(1&~(h>>4));
+#else
+  h += 9*(1&(h>>6));
+#endif
+  return (u8)(h & 0xf);
+}
+
+/*
+** Convert a 4-byte hex string into an integer
+*/
+static u32 jsonHexToInt4(const char *z){
+  u32 v;
+  assert( safe_isxdigit(z[0]) );
+  assert( safe_isxdigit(z[1]) );
+  assert( safe_isxdigit(z[2]) );
+  assert( safe_isxdigit(z[3]) );
+  v = (jsonHexToInt(z[0])<<12)
+    + (jsonHexToInt(z[1])<<8)
+    + (jsonHexToInt(z[2])<<4)
+    + jsonHexToInt(z[3]);
+  return v;
+}
+
 /*
 ** Make the JsonNode the return value of the function.
 */
@@ -179224,15 +184053,8 @@ static void jsonReturn(
           }else{
             c = z[++i];
             if( c=='u' ){
-              u32 v = 0, k;
-              for(k=0; k<4; i++, k++){
-                assert( i<n-2 );
-                c = z[i+1];
-                assert( safe_isxdigit(c) );
-                if( c<='9' ) v = v*16 + c - '0';
-                else if( c<='F' ) v = v*16 + c - 'A' + 10;
-                else v = v*16 + c - 'a' + 10;
-              }
+              u32 v = jsonHexToInt4(z+i+1);
+              i += 4;
               if( v==0 ) break;
               if( v<=0x7f ){
                 zOut[j++] = (char)v;
@@ -179240,9 +184062,25 @@ static void jsonReturn(
                 zOut[j++] = (char)(0xc0 | (v>>6));
                 zOut[j++] = 0x80 | (v&0x3f);
               }else{
-                zOut[j++] = (char)(0xe0 | (v>>12));
-                zOut[j++] = 0x80 | ((v>>6)&0x3f);
-                zOut[j++] = 0x80 | (v&0x3f);
+                u32 vlo;
+                if( (v&0xfc00)==0xd800
+                  && i<n-6
+                  && z[i+1]=='\\'
+                  && z[i+2]=='u'
+                  && ((vlo = jsonHexToInt4(z+i+3))&0xfc00)==0xdc00
+                ){
+                  /* We have a surrogate pair */
+                  v = ((v&0x3ff)<<10) + (vlo&0x3ff) + 0x10000;
+                  i += 6;
+                  zOut[j++] = 0xf0 | (v>>18);
+                  zOut[j++] = 0x80 | ((v>>12)&0x3f);
+                  zOut[j++] = 0x80 | ((v>>6)&0x3f);
+                  zOut[j++] = 0x80 | (v&0x3f);
+                }else{
+                  zOut[j++] = 0xe0 | (v>>12);
+                  zOut[j++] = 0x80 | ((v>>6)&0x3f);
+                  zOut[j++] = 0x80 | (v&0x3f);
+                }
               }
             }else{
               if( c=='b' ){
@@ -179692,6 +184530,7 @@ static JsonNode *jsonLookupStep(
   const char *zKey;
   JsonNode *pRoot = &pParse->aNode[iRoot];
   if( zPath[0]==0 ) return pRoot;
+  if( pRoot->jnFlags & JNODE_REPLACE ) return 0;
   if( zPath[0]=='.' ){
     if( pRoot->eType!=JSON_OBJECT ) return 0;
     zPath++;
@@ -179732,7 +184571,7 @@ static JsonNode *jsonLookupStep(
       u32 iStart, iLabel;
       JsonNode *pNode;
       iStart = jsonParseAddNode(pParse, JSON_OBJECT, 2, 0);
-      iLabel = jsonParseAddNode(pParse, JSON_STRING, i, zPath);
+      iLabel = jsonParseAddNode(pParse, JSON_STRING, nKey, zKey);
       zPath += i;
       pNode = jsonLookupAppend(pParse, zPath, pApnd, pzErr);
       if( pParse->oom ) return 0;
@@ -179744,18 +184583,49 @@ static JsonNode *jsonLookupStep(
       }
       return pNode;
     }
-  }else if( zPath[0]=='[' && safe_isdigit(zPath[1]) ){
-    if( pRoot->eType!=JSON_ARRAY ) return 0;
+  }else if( zPath[0]=='[' ){
     i = 0;
     j = 1;
     while( safe_isdigit(zPath[j]) ){
       i = i*10 + zPath[j] - '0';
       j++;
     }
-    if( zPath[j]!=']' ){
-      *pzErr = zPath;
-      return 0;
+    if( j<2 || zPath[j]!=']' ){
+      if( zPath[1]=='#' ){
+        JsonNode *pBase = pRoot;
+        int iBase = iRoot;
+        if( pRoot->eType!=JSON_ARRAY ) return 0;
+        for(;;){
+          while( j<=pBase->n ){
+            if( (pBase[j].jnFlags & JNODE_REMOVE)==0 ) i++;
+            j += jsonNodeSize(&pBase[j]);
+          }
+          if( (pBase->jnFlags & JNODE_APPEND)==0 ) break;
+          iBase += pBase->u.iAppend;
+          pBase = &pParse->aNode[iBase];
+          j = 1;
+        }
+        j = 2;
+        if( zPath[2]=='-' && safe_isdigit(zPath[3]) ){
+          unsigned int x = 0;
+          j = 3;
+          do{
+            x = x*10 + zPath[j] - '0';
+            j++;
+          }while( safe_isdigit(zPath[j]) );
+          if( x>i ) return 0;
+          i -= x;
+        }
+        if( zPath[j]!=']' ){
+          *pzErr = zPath;
+          return 0;
+        }
+      }else{
+        *pzErr = zPath;
+        return 0;
+      }
     }
+    if( pRoot->eType!=JSON_ARRAY ) return 0;
     zPath += j + 1;
     j = 1;
     for(;;){
@@ -180428,7 +185298,7 @@ static void jsonArrayStep(
     if( pStr->zBuf==0 ){
       jsonInit(pStr, ctx);
       jsonAppendChar(pStr, '[');
-    }else{
+    }else if( pStr->nUsed>1 ){
       jsonAppendChar(pStr, ',');
       pStr->pCtx = ctx;
     }
@@ -180476,9 +185346,11 @@ static void jsonGroupInverse(
   int argc,
   sqlite3_value **argv
 ){
-  int i;
+  unsigned int i;
   int inStr = 0;
+  int nNest = 0;
   char *z;
+  char c;
   JsonString *pStr;
   UNUSED_PARAM(argc);
   UNUSED_PARAM(argv);
@@ -180489,12 +185361,18 @@ static void jsonGroupInverse(
   if( NEVER(!pStr) ) return;
 #endif
   z = pStr->zBuf;
-  for(i=1; z[i]!=',' || inStr; i++){
-    assert( i<pStr->nUsed );
-    if( z[i]=='"' ){
+  for(i=1; (c = z[i])!=',' || inStr || nNest; i++){
+    if( i>=pStr->nUsed ){
+      pStr->nUsed = 1;
+      return;
+    }
+    if( c=='"' ){
       inStr = !inStr;
-    }else if( z[i]=='\\' ){
+    }else if( c=='\\' ){
       i++;
+    }else if( !inStr ){
+      if( c=='{' || c=='[' ) nNest++;
+      if( c=='}' || c==']' ) nNest--;
     }
   }
   pStr->nUsed -= i;      
@@ -180524,7 +185402,7 @@ static void jsonObjectStep(
     if( pStr->zBuf==0 ){
       jsonInit(pStr, ctx);
       jsonAppendChar(pStr, '{');
-    }else{
+    }else if( pStr->nUsed>1 ){
       jsonAppendChar(pStr, ',');
       pStr->pCtx = ctx;
     }
@@ -180620,6 +185498,7 @@ static int jsonEachConnect(
     pNew = *ppVtab = sqlite3_malloc( sizeof(*pNew) );
     if( pNew==0 ) return SQLITE_NOMEM;
     memset(pNew, 0, sizeof(*pNew));
+    sqlite3_vtab_config(db, SQLITE_VTAB_INNOCUOUS);
   }
   return rc;
 }
@@ -181110,16 +185989,19 @@ SQLITE_PRIVATE int sqlite3Json1Init(sqlite3 *db){
     { "json_tree",            &jsonTreeModule               },
   };
 #endif
+  static const int enc = 
+       SQLITE_UTF8 |
+       SQLITE_DETERMINISTIC |
+       SQLITE_INNOCUOUS;
   for(i=0; i<sizeof(aFunc)/sizeof(aFunc[0]) && rc==SQLITE_OK; i++){
-    rc = sqlite3_create_function(db, aFunc[i].zName, aFunc[i].nArg,
-                                 SQLITE_UTF8 | SQLITE_DETERMINISTIC, 
+    rc = sqlite3_create_function(db, aFunc[i].zName, aFunc[i].nArg, enc,
                                  (void*)&aFunc[i].flag,
                                  aFunc[i].xFunc, 0, 0);
   }
 #ifndef SQLITE_OMIT_WINDOWFUNC
   for(i=0; i<sizeof(aAgg)/sizeof(aAgg[0]) && rc==SQLITE_OK; i++){
     rc = sqlite3_create_window_function(db, aAgg[i].zName, aAgg[i].nArg,
-                                 SQLITE_UTF8 | SQLITE_DETERMINISTIC, 0,
+                                 SQLITE_SUBTYPE | enc, 0,
                                  aAgg[i].xStep, aAgg[i].xFinal,
                                  aAgg[i].xValue, jsonGroupInverse, 0);
   }
@@ -181215,10 +186097,7 @@ SQLITE_API int sqlite3_json_init(
 #else
 /*   #include "sqlite3.h" */
 #endif
-
-/* #include <string.h> */
-/* #include <assert.h> */
-/* #include <stdio.h> */
+SQLITE_PRIVATE int sqlite3GetToken(const unsigned char*,int*); /* In the SQLite core */
 
 #ifndef SQLITE_AMALGAMATION
 #include "sqlite3rtree.h"
@@ -181227,7 +186106,17 @@ typedef sqlite3_uint64 u64;
 typedef unsigned char u8;
 typedef unsigned short u16;
 typedef unsigned int u32;
+#if !defined(NDEBUG) && !defined(SQLITE_DEBUG)
+# define NDEBUG 1
 #endif
+#if defined(NDEBUG) && defined(SQLITE_DEBUG)
+# undef NDEBUG
+#endif
+#endif
+
+/* #include <string.h> */
+/* #include <stdio.h> */
+/* #include <assert.h> */
 
 /*  The following macro is used to suppress compiler warnings.
 */
@@ -181472,6 +186361,12 @@ struct RtreeConstraint {
 #define RTREE_MATCH 0x46  /* F: Old-style sqlite3_rtree_geometry_callback() */
 #define RTREE_QUERY 0x47  /* G: New-style sqlite3_rtree_query_callback() */
 
+/* Special operators available only on cursors.  Needs to be consecutive
+** with the normal values above, but must be less than RTREE_MATCH.  These
+** are used in the cursor for contraints such as x=NULL (RTREE_FALSE) or
+** x<'xyz' (RTREE_TRUE) */
+#define RTREE_TRUE  0x3f  /* ? */
+#define RTREE_FALSE 0x40  /* @ */
 
 /* 
 ** An rtree structure node.
@@ -181816,7 +186711,6 @@ static int nodeAcquire(
   ** increase its reference count and return it.
   */
   if( (pNode = nodeHashLookup(pRtree, iNode))!=0 ){
-    assert( !pParent || !pNode->pParent || pNode->pParent==pParent );
     if( pParent && !pNode->pParent ){
       if( nodeInParentChain(pNode, pParent) ){
         RTREE_IS_CORRUPT(pRtree);
@@ -181824,6 +186718,9 @@ static int nodeAcquire(
       }
       pParent->nRef++;
       pNode->pParent = pParent;
+    }else if( pParent && pNode->pParent && pParent!=pNode->pParent ){
+      RTREE_IS_CORRUPT(pRtree);
+      return SQLITE_CORRUPT_VTAB;
     }
     pNode->nRef++;
     *ppNode = pNode;
@@ -182203,9 +187100,12 @@ static int rtreeOpen(sqlite3_vtab *pVTab, sqlite3_vtab_cursor **ppCursor){
 
 
 /*
-** Free the RtreeCursor.aConstraint[] array and its contents.
+** Reset a cursor back to its initial state.
 */
-static void freeCursorConstraints(RtreeCursor *pCsr){
+static void resetCursor(RtreeCursor *pCsr){
+  Rtree *pRtree = (Rtree *)(pCsr->base.pVtab);
+  int ii;
+  sqlite3_stmt *pStmt;
   if( pCsr->aConstraint ){
     int i;                        /* Used to iterate through constraint array */
     for(i=0; i<pCsr->nConstraint; i++){
@@ -182218,6 +187118,13 @@ static void freeCursorConstraints(RtreeCursor *pCsr){
     sqlite3_free(pCsr->aConstraint);
     pCsr->aConstraint = 0;
   }
+  for(ii=0; ii<RTREE_CACHE_SZ; ii++) nodeRelease(pRtree, pCsr->aNode[ii]);
+  sqlite3_free(pCsr->aPoint);
+  pStmt = pCsr->pReadAux;
+  memset(pCsr, 0, sizeof(RtreeCursor));
+  pCsr->base.pVtab = (sqlite3_vtab*)pRtree;
+  pCsr->pReadAux = pStmt;
+
 }
 
 /* 
@@ -182225,13 +187132,10 @@ static void freeCursorConstraints(RtreeCursor *pCsr){
 */
 static int rtreeClose(sqlite3_vtab_cursor *cur){
   Rtree *pRtree = (Rtree *)(cur->pVtab);
-  int ii;
   RtreeCursor *pCsr = (RtreeCursor *)cur;
   assert( pRtree->nCursor>0 );
-  freeCursorConstraints(pCsr);
+  resetCursor(pCsr);
   sqlite3_finalize(pCsr->pReadAux);
-  sqlite3_free(pCsr->aPoint);
-  for(ii=0; ii<RTREE_CACHE_SZ; ii++) nodeRelease(pRtree, pCsr->aNode[ii]);
   sqlite3_free(pCsr);
   pRtree->nCursor--;
   nodeBlobReset(pRtree);
@@ -182389,9 +187293,12 @@ static void rtreeNonleafConstraint(
   pCellData += 8 + 4*(p->iCoord&0xfe);
 
   assert(p->op==RTREE_LE || p->op==RTREE_LT || p->op==RTREE_GE 
-      || p->op==RTREE_GT || p->op==RTREE_EQ );
+      || p->op==RTREE_GT || p->op==RTREE_EQ || p->op==RTREE_TRUE
+      || p->op==RTREE_FALSE );
   assert( ((((char*)pCellData) - (char*)0)&3)==0 );  /* 4-byte aligned */
   switch( p->op ){
+    case RTREE_TRUE:  return;   /* Always satisfied */
+    case RTREE_FALSE: break;    /* Never satisfied */
     case RTREE_LE:
     case RTREE_LT:
     case RTREE_EQ:
@@ -182429,16 +187336,19 @@ static void rtreeLeafConstraint(
   RtreeDValue xN;      /* Coordinate value converted to a double */
 
   assert(p->op==RTREE_LE || p->op==RTREE_LT || p->op==RTREE_GE 
-      || p->op==RTREE_GT || p->op==RTREE_EQ );
+      || p->op==RTREE_GT || p->op==RTREE_EQ || p->op==RTREE_TRUE
+      || p->op==RTREE_FALSE );
   pCellData += 8 + p->iCoord*4;
   assert( ((((char*)pCellData) - (char*)0)&3)==0 );  /* 4-byte aligned */
   RTREE_DECODE_COORD(eInt, pCellData, xN);
   switch( p->op ){
-    case RTREE_LE: if( xN <= p->u.rValue ) return;  break;
-    case RTREE_LT: if( xN <  p->u.rValue ) return;  break;
-    case RTREE_GE: if( xN >= p->u.rValue ) return;  break;
-    case RTREE_GT: if( xN >  p->u.rValue ) return;  break;
-    default:       if( xN == p->u.rValue ) return;  break;
+    case RTREE_TRUE:  return;   /* Always satisfied */
+    case RTREE_FALSE: break;    /* Never satisfied */
+    case RTREE_LE:    if( xN <= p->u.rValue ) return;  break;
+    case RTREE_LT:    if( xN <  p->u.rValue ) return;  break;
+    case RTREE_GE:    if( xN >= p->u.rValue ) return;  break;
+    case RTREE_GT:    if( xN >  p->u.rValue ) return;  break;
+    default:          if( xN == p->u.rValue ) return;  break;
   }
   *peWithin = NOT_WITHIN;
 }
@@ -182711,13 +187621,14 @@ static int rtreeStepToLeaf(RtreeCursor *pCur){
 
   eInt = pRtree->eCoordType==RTREE_COORD_INT32;
   while( (p = rtreeSearchPointFirst(pCur))!=0 && p->iLevel>0 ){
+    u8 *pCellData;
     pNode = rtreeNodeOfFirstSearchPoint(pCur, &rc);
     if( rc ) return rc;
     nCell = NCELL(pNode);
     assert( nCell<200 );
+    pCellData = pNode->zData + (4+pRtree->nBytesPerCell*p->iCell);
     while( p->iCell<nCell ){
       sqlite3_rtree_dbl rScore = (sqlite3_rtree_dbl)-1;
-      u8 *pCellData = pNode->zData + (4+pRtree->nBytesPerCell*p->iCell);
       eWithin = FULLY_WITHIN;
       for(ii=0; ii<nConstraint; ii++){
         RtreeConstraint *pConstraint = pCur->aConstraint + ii;
@@ -182730,13 +187641,23 @@ static int rtreeStepToLeaf(RtreeCursor *pCur){
         }else{
           rtreeNonleafConstraint(pConstraint, eInt, pCellData, &eWithin);
         }
-        if( eWithin==NOT_WITHIN ) break;
+        if( eWithin==NOT_WITHIN ){
+          p->iCell++;
+          pCellData += pRtree->nBytesPerCell;
+          break;
+        }
       }
-      p->iCell++;
       if( eWithin==NOT_WITHIN ) continue;
+      p->iCell++;
       x.iLevel = p->iLevel - 1;
       if( x.iLevel ){
         x.id = readInt64(pCellData);
+        for(ii=0; ii<pCur->nPoint; ii++){
+          if( pCur->aPoint[ii].id==x.id ){
+            RTREE_IS_CORRUPT(pRtree);
+            return SQLITE_CORRUPT_VTAB;
+          }
+        }
         x.iCell = 0;
       }else{
         x.id = p->id;
@@ -182920,17 +187841,11 @@ static int rtreeFilter(
   int ii;
   int rc = SQLITE_OK;
   int iCell = 0;
-  sqlite3_stmt *pStmt;
 
   rtreeReference(pRtree);
 
   /* Reset the cursor to the same state as rtreeOpen() leaves it in. */
-  freeCursorConstraints(pCsr);
-  sqlite3_free(pCsr->aPoint);
-  pStmt = pCsr->pReadAux;
-  memset(pCsr, 0, sizeof(RtreeCursor));
-  pCsr->base.pVtab = (sqlite3_vtab*)pRtree;
-  pCsr->pReadAux = pStmt;
+  resetCursor(pCsr);
 
   pCsr->iStrategy = idxNum;
   if( idxNum==1 ){
@@ -182939,7 +187854,15 @@ static int rtreeFilter(
     RtreeSearchPoint *p;     /* Search point for the leaf */
     i64 iRowid = sqlite3_value_int64(argv[0]);
     i64 iNode = 0;
-    rc = findLeafNode(pRtree, iRowid, &pLeaf, &iNode);
+    int eType = sqlite3_value_numeric_type(argv[0]);
+    if( eType==SQLITE_INTEGER
+     || (eType==SQLITE_FLOAT && sqlite3_value_double(argv[0])==iRowid)
+    ){
+      rc = findLeafNode(pRtree, iRowid, &pLeaf, &iNode);
+    }else{
+      rc = SQLITE_OK;
+      pLeaf = 0;
+    }
     if( rc==SQLITE_OK && pLeaf!=0 ){
       p = rtreeSearchPointNew(pCsr, RTREE_ZERO, 0);
       assert( p!=0 );  /* Always returns pCsr->sPoint */
@@ -182969,6 +187892,7 @@ static int rtreeFilter(
                 || (idxStr && (int)strlen(idxStr)==argc*2) );
         for(ii=0; ii<argc; ii++){
           RtreeConstraint *p = &pCsr->aConstraint[ii];
+          int eType = sqlite3_value_numeric_type(argv[ii]);
           p->op = idxStr[ii*2];
           p->iCoord = idxStr[ii*2+1]-'0';
           if( p->op>=RTREE_MATCH ){
@@ -182983,12 +187907,21 @@ static int rtreeFilter(
             p->pInfo->nCoord = pRtree->nDim2;
             p->pInfo->anQueue = pCsr->anQueue;
             p->pInfo->mxLevel = pRtree->iDepth + 1;
-          }else{
+          }else if( eType==SQLITE_INTEGER || eType==SQLITE_FLOAT ){
 #ifdef SQLITE_RTREE_INT_ONLY
             p->u.rValue = sqlite3_value_int64(argv[ii]);
 #else
             p->u.rValue = sqlite3_value_double(argv[ii]);
 #endif
+          }else{
+            p->u.rValue = RTREE_ZERO;
+            if( eType==SQLITE_NULL ){
+              p->op = RTREE_FALSE;
+            }else if( p->op==RTREE_LT || p->op==RTREE_LE ){
+              p->op = RTREE_TRUE;
+            }else{
+              p->op = RTREE_FALSE;
+            }
           }
         }
       }
@@ -184765,6 +189698,14 @@ static int getNodeSize(
   return rc;
 }
 
+/*
+** Return the length of a token
+*/
+static int rtreeTokenLength(const char *z){
+  int dummy = 0;
+  return sqlite3GetToken((const unsigned char*)z,&dummy);
+}
+
 /* 
 ** This function is the implementation of both the xConnect and xCreate
 ** methods of the r-tree virtual table.
@@ -184801,8 +189742,8 @@ static int rtreeInit(
   };
 
   assert( RTREE_MAX_AUX_COLUMN<256 ); /* Aux columns counted by a u8 */
-  if( argc>RTREE_MAX_AUX_COLUMN+3 ){
-    *pzErr = sqlite3_mprintf("%s", aErrMsg[3]);
+  if( argc<6 || argc>RTREE_MAX_AUX_COLUMN+3 ){
+    *pzErr = sqlite3_mprintf("%s", aErrMsg[2 + (argc>=6)]);
     return SQLITE_ERROR;
   }
 
@@ -184830,16 +189771,18 @@ static int rtreeInit(
   ** the r-tree table schema.
   */
   pSql = sqlite3_str_new(db);
-  sqlite3_str_appendf(pSql, "CREATE TABLE x(%s", argv[3]);
+  sqlite3_str_appendf(pSql, "CREATE TABLE x(%.*s INT", 
+                      rtreeTokenLength(argv[3]), argv[3]);
   for(ii=4; ii<argc; ii++){
-    if( argv[ii][0]=='+' ){
+    const char *zArg = argv[ii];
+    if( zArg[0]=='+' ){
       pRtree->nAux++;
-      sqlite3_str_appendf(pSql, ",%s", argv[ii]+1);
+      sqlite3_str_appendf(pSql, ",%.*s", rtreeTokenLength(zArg+1), zArg+1);
     }else if( pRtree->nAux>0 ){
       break;
     }else{
       pRtree->nDim2++;
-      sqlite3_str_appendf(pSql, ",%s", argv[ii]);
+      sqlite3_str_appendf(pSql, ",%.*s NUM", rtreeTokenLength(zArg), zArg);
     }
   }
   sqlite3_str_appendf(pSql, ");");
@@ -186787,17 +191730,11 @@ static int geopolyFilter(
   RtreeNode *pRoot = 0;
   int rc = SQLITE_OK;
   int iCell = 0;
-  sqlite3_stmt *pStmt;
 
   rtreeReference(pRtree);
 
   /* Reset the cursor to the same state as rtreeOpen() leaves it in. */
-  freeCursorConstraints(pCsr);
-  sqlite3_free(pCsr->aPoint);
-  pStmt = pCsr->pReadAux;
-  memset(pCsr, 0, sizeof(RtreeCursor));
-  pCsr->base.pVtab = (sqlite3_vtab*)pRtree;
-  pCsr->pReadAux = pStmt;
+  resetCursor(pCsr);
 
   pCsr->iStrategy = idxNum;
   if( idxNum==1 ){
@@ -187234,14 +192171,20 @@ static int sqlite3_geopoly_init(sqlite3 *db){
   };
   int i;
   for(i=0; i<sizeof(aFunc)/sizeof(aFunc[0]) && rc==SQLITE_OK; i++){
-    int enc = aFunc[i].bPure ? SQLITE_UTF8|SQLITE_DETERMINISTIC : SQLITE_UTF8;
+    int enc;
+    if( aFunc[i].bPure ){
+      enc = SQLITE_UTF8|SQLITE_DETERMINISTIC|SQLITE_INNOCUOUS;
+    }else{
+      enc = SQLITE_UTF8|SQLITE_DIRECTONLY;
+    }
     rc = sqlite3_create_function(db, aFunc[i].zName, aFunc[i].nArg,
                                  enc, 0,
                                  aFunc[i].xFunc, 0, 0);
   }
   for(i=0; i<sizeof(aAgg)/sizeof(aAgg[0]) && rc==SQLITE_OK; i++){
-    rc = sqlite3_create_function(db, aAgg[i].zName, 1, SQLITE_UTF8, 0,
-                                 0, aAgg[i].xStep, aAgg[i].xFinal);
+    rc = sqlite3_create_function(db, aAgg[i].zName, 1, 
+              SQLITE_UTF8|SQLITE_DETERMINISTIC|SQLITE_INNOCUOUS, 0,
+              0, aAgg[i].xStep, aAgg[i].xFinal);
   }
   if( rc==SQLITE_OK ){
     rc = sqlite3_create_module_v2(db, "geopoly", &geopolyModule, 0, 0);
@@ -187931,26 +192874,27 @@ static void icuLoadCollation(
 ** Register the ICU extension functions with database db.
 */
 SQLITE_PRIVATE int sqlite3IcuInit(sqlite3 *db){
+# define SQLITEICU_EXTRAFLAGS (SQLITE_DETERMINISTIC|SQLITE_INNOCUOUS)
   static const struct IcuScalar {
     const char *zName;                        /* Function name */
     unsigned char nArg;                       /* Number of arguments */
-    unsigned short enc;                       /* Optimal text encoding */
+    unsigned int enc;                         /* Optimal text encoding */
     unsigned char iContext;                   /* sqlite3_user_data() context */
     void (*xFunc)(sqlite3_context*,int,sqlite3_value**);
   } scalars[] = {
-    {"icu_load_collation",  2, SQLITE_UTF8,                1, icuLoadCollation},
+    {"icu_load_collation",2,SQLITE_UTF8|SQLITE_DIRECTONLY,1, icuLoadCollation},
 #if !defined(SQLITE_CORE) || defined(SQLITE_ENABLE_ICU)
-    {"regexp", 2, SQLITE_ANY|SQLITE_DETERMINISTIC,         0, icuRegexpFunc},
-    {"lower",  1, SQLITE_UTF16|SQLITE_DETERMINISTIC,       0, icuCaseFunc16},
-    {"lower",  2, SQLITE_UTF16|SQLITE_DETERMINISTIC,       0, icuCaseFunc16},
-    {"upper",  1, SQLITE_UTF16|SQLITE_DETERMINISTIC,       1, icuCaseFunc16},
-    {"upper",  2, SQLITE_UTF16|SQLITE_DETERMINISTIC,       1, icuCaseFunc16},
-    {"lower",  1, SQLITE_UTF8|SQLITE_DETERMINISTIC,        0, icuCaseFunc16},
-    {"lower",  2, SQLITE_UTF8|SQLITE_DETERMINISTIC,        0, icuCaseFunc16},
-    {"upper",  1, SQLITE_UTF8|SQLITE_DETERMINISTIC,        1, icuCaseFunc16},
-    {"upper",  2, SQLITE_UTF8|SQLITE_DETERMINISTIC,        1, icuCaseFunc16},
-    {"like",   2, SQLITE_UTF8|SQLITE_DETERMINISTIC,        0, icuLikeFunc},
-    {"like",   3, SQLITE_UTF8|SQLITE_DETERMINISTIC,        0, icuLikeFunc},
+    {"regexp", 2, SQLITE_ANY|SQLITEICU_EXTRAFLAGS,         0, icuRegexpFunc},
+    {"lower",  1, SQLITE_UTF16|SQLITEICU_EXTRAFLAGS,       0, icuCaseFunc16},
+    {"lower",  2, SQLITE_UTF16|SQLITEICU_EXTRAFLAGS,       0, icuCaseFunc16},
+    {"upper",  1, SQLITE_UTF16|SQLITEICU_EXTRAFLAGS,       1, icuCaseFunc16},
+    {"upper",  2, SQLITE_UTF16|SQLITEICU_EXTRAFLAGS,       1, icuCaseFunc16},
+    {"lower",  1, SQLITE_UTF8|SQLITEICU_EXTRAFLAGS,        0, icuCaseFunc16},
+    {"lower",  2, SQLITE_UTF8|SQLITEICU_EXTRAFLAGS,        0, icuCaseFunc16},
+    {"upper",  1, SQLITE_UTF8|SQLITEICU_EXTRAFLAGS,        1, icuCaseFunc16},
+    {"upper",  2, SQLITE_UTF8|SQLITEICU_EXTRAFLAGS,        1, icuCaseFunc16},
+    {"like",   2, SQLITE_UTF8|SQLITEICU_EXTRAFLAGS,        0, icuLikeFunc},
+    {"like",   3, SQLITE_UTF8|SQLITEICU_EXTRAFLAGS,        0, icuLikeFunc},
 #endif /* !defined(SQLITE_CORE) || defined(SQLITE_ENABLE_ICU) */
   };
   int rc = SQLITE_OK;
@@ -189044,6 +193988,7 @@ SQLITE_API void sqlite3rbu_destroy_vfs(const char *zName);
 typedef struct RbuFrame RbuFrame;
 typedef struct RbuObjIter RbuObjIter;
 typedef struct RbuState RbuState;
+typedef struct RbuSpan RbuSpan;
 typedef struct rbu_vfs rbu_vfs;
 typedef struct rbu_file rbu_file;
 typedef struct RbuUpdateStmt RbuUpdateStmt;
@@ -189088,6 +194033,11 @@ struct RbuUpdateStmt {
   RbuUpdateStmt *pNext;
 };
 
+struct RbuSpan {
+  const char *zSpan;
+  int nSpan;
+};
+
 /*
 ** An iterator of this type is used to iterate through all objects in
 ** the target database that require updating. For each such table, the
@@ -189137,6 +194087,9 @@ struct RbuObjIter {
   sqlite3_stmt *pInsert;          /* Statement for INSERT operations */
   sqlite3_stmt *pDelete;          /* Statement for DELETE ops */
   sqlite3_stmt *pTmpInsert;       /* Insert into rbu_tmp_$zDataTbl */
+  int nIdxCol;
+  RbuSpan *aIdxCol;
+  char *zIdxSql;
 
   /* Last UPDATE used (for PK b-tree updates only), or NULL. */
   RbuUpdateStmt *pRbuUpdate;
@@ -189671,6 +194624,8 @@ static void rbuObjIterClearStatements(RbuObjIter *pIter){
     sqlite3_free(pUp);
     pUp = pTmp;
   }
+  sqlite3_free(pIter->aIdxCol);
+  sqlite3_free(pIter->zIdxSql);
   
   pIter->pSelect = 0;
   pIter->pInsert = 0;
@@ -189678,6 +194633,9 @@ static void rbuObjIterClearStatements(RbuObjIter *pIter){
   pIter->pRbuUpdate = 0;
   pIter->pTmpInsert = 0;
   pIter->nCol = 0;
+  pIter->nIdxCol = 0;
+  pIter->aIdxCol = 0;
+  pIter->zIdxSql = 0;
 }
 
 /*
@@ -189792,6 +194750,7 @@ static void rbuTargetNameFunc(
   zIn = (const char*)sqlite3_value_text(argv[0]);
   if( zIn ){
     if( rbuIsVacuum(p) ){
+      assert( argc==2 || argc==1 );
       if( argc==1 || 0==sqlite3_value_int(argv[1]) ){
         sqlite3_result_text(pCtx, zIn, -1, SQLITE_STATIC);
       }
@@ -189950,14 +194909,15 @@ static void rbuAllocateIterArrays(sqlite3rbu *p, RbuObjIter *pIter, int nCol){
 static char *rbuStrndup(const char *zStr, int *pRc){
   char *zRet = 0;
 
-  assert( *pRc==SQLITE_OK );
-  if( zStr ){
-    size_t nCopy = strlen(zStr) + 1;
-    zRet = (char*)sqlite3_malloc64(nCopy);
-    if( zRet ){
-      memcpy(zRet, zStr, nCopy);
-    }else{
-      *pRc = SQLITE_NOMEM;
+  if( *pRc==SQLITE_OK ){
+    if( zStr ){
+      size_t nCopy = strlen(zStr) + 1;
+      zRet = (char*)sqlite3_malloc64(nCopy);
+      if( zRet ){
+        memcpy(zRet, zStr, nCopy);
+      }else{
+        *pRc = SQLITE_NOMEM;
+      }
     }
   }
 
@@ -190129,6 +195089,9 @@ static void rbuObjIterCacheIndexedCols(sqlite3rbu *p, RbuObjIter *pIter){
     while( p->rc==SQLITE_OK && SQLITE_ROW==sqlite3_step(pXInfo) ){
       int iCid = sqlite3_column_int(pXInfo, 1);
       if( iCid>=0 ) pIter->abIndexed[iCid] = 1;
+      if( iCid==-2 ){
+        memset(pIter->abIndexed, 0x01, sizeof(u8)*pIter->nTblCol);
+      }
     }
     rbuFinalize(p, pXInfo);
     bIndex = 1;
@@ -190243,7 +195206,8 @@ static int rbuObjIterCacheTableInfo(sqlite3rbu *p, RbuObjIter *pIter){
         }
 
         pIter->azTblType[iOrder] = rbuStrndup(zType, &p->rc);
-        pIter->abTblPk[iOrder] = (iPk!=0);
+        assert( iPk>=0 );
+        pIter->abTblPk[iOrder] = (u8)iPk;
         pIter->abNotNull[iOrder] = (u8)bNotNull || (iPk!=0);
         iOrder++;
       }
@@ -190278,6 +195242,213 @@ static char *rbuObjIterGetCollist(
   return zList;
 }
 
+/*
+** Return a comma separated list of the quoted PRIMARY KEY column names,
+** in order, for the current table. Before each column name, add the text
+** zPre. After each column name, add the zPost text. Use zSeparator as
+** the separator text (usually ", ").
+*/
+static char *rbuObjIterGetPkList(
+  sqlite3rbu *p,                  /* RBU object */
+  RbuObjIter *pIter,              /* Object iterator for column names */
+  const char *zPre,               /* Before each quoted column name */
+  const char *zSeparator,         /* Separator to use between columns */
+  const char *zPost               /* After each quoted column name */
+){
+  int iPk = 1;
+  char *zRet = 0;
+  const char *zSep = "";
+  while( 1 ){
+    int i;
+    for(i=0; i<pIter->nTblCol; i++){
+      if( (int)pIter->abTblPk[i]==iPk ){
+        const char *zCol = pIter->azTblCol[i];
+        zRet = rbuMPrintf(p, "%z%s%s\"%w\"%s", zRet, zSep, zPre, zCol, zPost);
+        zSep = zSeparator;
+        break;
+      }
+    }
+    if( i==pIter->nTblCol ) break;
+    iPk++;
+  }
+  return zRet;
+}
+
+/*
+** This function is called as part of restarting an RBU vacuum within 
+** stage 1 of the process (while the *-oal file is being built) while
+** updating a table (not an index). The table may be a rowid table or
+** a WITHOUT ROWID table. It queries the target database to find the 
+** largest key that has already been written to the target table and
+** constructs a WHERE clause that can be used to extract the remaining
+** rows from the source table. For a rowid table, the WHERE clause
+** is of the form:
+**
+**     "WHERE _rowid_ > ?"
+**
+** and for WITHOUT ROWID tables:
+**
+**     "WHERE (key1, key2) > (?, ?)"
+**
+** Instead of "?" placeholders, the actual WHERE clauses created by
+** this function contain literal SQL values.
+*/
+static char *rbuVacuumTableStart(
+  sqlite3rbu *p,                  /* RBU handle */
+  RbuObjIter *pIter,              /* RBU iterator object */
+  int bRowid,                     /* True for a rowid table */
+  const char *zWrite              /* Target table name prefix */
+){
+  sqlite3_stmt *pMax = 0;
+  char *zRet = 0;
+  if( bRowid ){
+    p->rc = prepareFreeAndCollectError(p->dbMain, &pMax, &p->zErrmsg, 
+        sqlite3_mprintf(
+          "SELECT max(_rowid_) FROM \"%s%w\"", zWrite, pIter->zTbl
+        )
+    );
+    if( p->rc==SQLITE_OK && SQLITE_ROW==sqlite3_step(pMax) ){
+      sqlite3_int64 iMax = sqlite3_column_int64(pMax, 0);
+      zRet = rbuMPrintf(p, " WHERE _rowid_ > %lld ", iMax);
+    }
+    rbuFinalize(p, pMax);
+  }else{
+    char *zOrder = rbuObjIterGetPkList(p, pIter, "", ", ", " DESC");
+    char *zSelect = rbuObjIterGetPkList(p, pIter, "quote(", "||','||", ")");
+    char *zList = rbuObjIterGetPkList(p, pIter, "", ", ", "");
+
+    if( p->rc==SQLITE_OK ){
+      p->rc = prepareFreeAndCollectError(p->dbMain, &pMax, &p->zErrmsg, 
+          sqlite3_mprintf(
+            "SELECT %s FROM \"%s%w\" ORDER BY %s LIMIT 1", 
+                zSelect, zWrite, pIter->zTbl, zOrder
+          )
+      );
+      if( p->rc==SQLITE_OK && SQLITE_ROW==sqlite3_step(pMax) ){
+        const char *zVal = (const char*)sqlite3_column_text(pMax, 0);
+        zRet = rbuMPrintf(p, " WHERE (%s) > (%s) ", zList, zVal);
+      }
+      rbuFinalize(p, pMax);
+    }
+
+    sqlite3_free(zOrder);
+    sqlite3_free(zSelect);
+    sqlite3_free(zList);
+  }
+  return zRet;
+}
+
+/*
+** This function is called as part of restating an RBU vacuum when the
+** current operation is writing content to an index. If possible, it
+** queries the target index b-tree for the largest key already written to
+** it, then composes and returns an expression that can be used in a WHERE 
+** clause to select the remaining required rows from the source table. 
+** It is only possible to return such an expression if:
+**
+**   * The index contains no DESC columns, and
+**   * The last key written to the index before the operation was 
+**     suspended does not contain any NULL values.
+**
+** The expression is of the form:
+**
+**   (index-field1, index-field2, ...) > (?, ?, ...)
+**
+** except that the "?" placeholders are replaced with literal values.
+**
+** If the expression cannot be created, NULL is returned. In this case,
+** the caller has to use an OFFSET clause to extract only the required 
+** rows from the sourct table, just as it does for an RBU update operation.
+*/
+char *rbuVacuumIndexStart(
+  sqlite3rbu *p,                  /* RBU handle */
+  RbuObjIter *pIter               /* RBU iterator object */
+){
+  char *zOrder = 0;
+  char *zLhs = 0;
+  char *zSelect = 0;
+  char *zVector = 0;
+  char *zRet = 0;
+  int bFailed = 0;
+  const char *zSep = "";
+  int iCol = 0;
+  sqlite3_stmt *pXInfo = 0;
+
+  p->rc = prepareFreeAndCollectError(p->dbMain, &pXInfo, &p->zErrmsg,
+      sqlite3_mprintf("PRAGMA main.index_xinfo = %Q", pIter->zIdx)
+  );
+  while( p->rc==SQLITE_OK && SQLITE_ROW==sqlite3_step(pXInfo) ){
+    int iCid = sqlite3_column_int(pXInfo, 1);
+    const char *zCollate = (const char*)sqlite3_column_text(pXInfo, 4);
+    const char *zCol;
+    if( sqlite3_column_int(pXInfo, 3) ){
+      bFailed = 1;
+      break;
+    }
+
+    if( iCid<0 ){
+      if( pIter->eType==RBU_PK_IPK ){
+        int i;
+        for(i=0; pIter->abTblPk[i]==0; i++);
+        assert( i<pIter->nTblCol );
+        zCol = pIter->azTblCol[i];
+      }else{
+        zCol = "_rowid_";
+      }
+    }else{
+      zCol = pIter->azTblCol[iCid];
+    }
+
+    zLhs = rbuMPrintf(p, "%z%s \"%w\" COLLATE %Q",
+        zLhs, zSep, zCol, zCollate
+        );
+    zOrder = rbuMPrintf(p, "%z%s \"rbu_imp_%d%w\" COLLATE %Q DESC",
+        zOrder, zSep, iCol, zCol, zCollate
+        );
+    zSelect = rbuMPrintf(p, "%z%s quote(\"rbu_imp_%d%w\")",
+        zSelect, zSep, iCol, zCol
+        );
+    zSep = ", ";
+    iCol++;
+  }
+  rbuFinalize(p, pXInfo);
+  if( bFailed ) goto index_start_out;
+
+  if( p->rc==SQLITE_OK ){
+    sqlite3_stmt *pSel = 0;
+
+    p->rc = prepareFreeAndCollectError(p->dbMain, &pSel, &p->zErrmsg,
+        sqlite3_mprintf("SELECT %s FROM \"rbu_imp_%w\" ORDER BY %s LIMIT 1",
+          zSelect, pIter->zTbl, zOrder
+        )
+    );
+    if( p->rc==SQLITE_OK && SQLITE_ROW==sqlite3_step(pSel) ){
+      zSep = "";
+      for(iCol=0; iCol<pIter->nCol; iCol++){
+        const char *zQuoted = (const char*)sqlite3_column_text(pSel, iCol);
+        if( zQuoted[0]=='N' ){
+          bFailed = 1;
+          break;
+        }
+        zVector = rbuMPrintf(p, "%z%s%s", zVector, zSep, zQuoted);
+        zSep = ", ";
+      }
+
+      if( !bFailed ){
+        zRet = rbuMPrintf(p, "(%s) > (%s)", zLhs, zVector);
+      }
+    }
+    rbuFinalize(p, pSel);
+  }
+
+ index_start_out:
+  sqlite3_free(zOrder);
+  sqlite3_free(zSelect);
+  sqlite3_free(zVector);
+  sqlite3_free(zLhs);
+  return zRet;
+}
+
 /*
 ** This function is used to create a SELECT list (the list of SQL 
 ** expressions that follows a SELECT keyword) for a SELECT statement 
@@ -190332,29 +195503,37 @@ static char *rbuObjIterGetIndexCols(
     int iCid = sqlite3_column_int(pXInfo, 1);
     int bDesc = sqlite3_column_int(pXInfo, 3);
     const char *zCollate = (const char*)sqlite3_column_text(pXInfo, 4);
-    const char *zCol;
+    const char *zCol = 0;
     const char *zType;
 
-    if( iCid<0 ){
-      /* An integer primary key. If the table has an explicit IPK, use
-      ** its name. Otherwise, use "rbu_rowid".  */
-      if( pIter->eType==RBU_PK_IPK ){
-        int i;
-        for(i=0; pIter->abTblPk[i]==0; i++);
-        assert( i<pIter->nTblCol );
-        zCol = pIter->azTblCol[i];
-      }else if( rbuIsVacuum(p) ){
-        zCol = "_rowid_";
+    if( iCid==-2 ){
+      int iSeq = sqlite3_column_int(pXInfo, 0);
+      zRet = sqlite3_mprintf("%z%s(%.*s) COLLATE %Q", zRet, zCom,
+          pIter->aIdxCol[iSeq].nSpan, pIter->aIdxCol[iSeq].zSpan, zCollate
+      );
+      zType = "";
+    }else {
+      if( iCid<0 ){
+        /* An integer primary key. If the table has an explicit IPK, use
+        ** its name. Otherwise, use "rbu_rowid".  */
+        if( pIter->eType==RBU_PK_IPK ){
+          int i;
+          for(i=0; pIter->abTblPk[i]==0; i++);
+          assert( i<pIter->nTblCol );
+          zCol = pIter->azTblCol[i];
+        }else if( rbuIsVacuum(p) ){
+          zCol = "_rowid_";
+        }else{
+          zCol = "rbu_rowid";
+        }
+        zType = "INTEGER";
       }else{
-        zCol = "rbu_rowid";
+        zCol = pIter->azTblCol[iCid];
+        zType = pIter->azTblType[iCid];
       }
-      zType = "INTEGER";
-    }else{
-      zCol = pIter->azTblCol[iCid];
-      zType = pIter->azTblType[iCid];
+      zRet = sqlite3_mprintf("%z%s\"%w\" COLLATE %Q", zRet, zCom,zCol,zCollate);
     }
 
-    zRet = sqlite3_mprintf("%z%s\"%w\" COLLATE %Q", zRet, zCom, zCol, zCollate);
     if( pIter->bUnique==0 || sqlite3_column_int(pXInfo, 5) ){
       const char *zOrder = (bDesc ? " DESC" : "");
       zImpPK = sqlite3_mprintf("%z%s\"rbu_imp_%d%w\"%s", 
@@ -190834,6 +196013,8 @@ static char *rbuObjIterGetIndexWhere(sqlite3rbu *p, RbuObjIter *pIter){
   int rc = p->rc;
   char *zRet = 0;
 
+  assert( pIter->zIdxSql==0 && pIter->nIdxCol==0 && pIter->aIdxCol==0 );
+
   if( rc==SQLITE_OK ){
     rc = prepareAndCollectError(p->dbMain, &pStmt, &p->zErrmsg,
         "SELECT trim(sql) FROM sqlite_master WHERE type='index' AND name=?"
@@ -190843,21 +196024,50 @@ static char *rbuObjIterGetIndexWhere(sqlite3rbu *p, RbuObjIter *pIter){
     int rc2;
     rc = sqlite3_bind_text(pStmt, 1, pIter->zIdx, -1, SQLITE_STATIC);
     if( rc==SQLITE_OK && SQLITE_ROW==sqlite3_step(pStmt) ){
-      const char *zSql = (const char*)sqlite3_column_text(pStmt, 0);
+      char *zSql = (char*)sqlite3_column_text(pStmt, 0);
+      if( zSql ){
+        pIter->zIdxSql = zSql = rbuStrndup(zSql, &rc);
+      }
       if( zSql ){
         int nParen = 0;           /* Number of open parenthesis */
         int i;
+        int iIdxCol = 0;
+        int nIdxAlloc = 0;
         for(i=0; zSql[i]; i++){
           char c = zSql[i];
+
+          /* If necessary, grow the pIter->aIdxCol[] array */
+          if( iIdxCol==nIdxAlloc ){
+            RbuSpan *aIdxCol = (RbuSpan*)sqlite3_realloc(
+                pIter->aIdxCol, (nIdxAlloc+16)*sizeof(RbuSpan)
+            );
+            if( aIdxCol==0 ){
+              rc = SQLITE_NOMEM;
+              break;
+            }
+            pIter->aIdxCol = aIdxCol;
+            nIdxAlloc += 16;
+          }
+
           if( c=='(' ){
+            if( nParen==0 ){
+              assert( iIdxCol==0 );
+              pIter->aIdxCol[0].zSpan = &zSql[i+1];
+            }
             nParen++;
           }
           else if( c==')' ){
             nParen--;
             if( nParen==0 ){
+              int nSpan = &zSql[i] - pIter->aIdxCol[iIdxCol].zSpan;
+              pIter->aIdxCol[iIdxCol++].nSpan = nSpan;
               i++;
               break;
             }
+          }else if( c==',' && nParen==1 ){
+            int nSpan = &zSql[i] - pIter->aIdxCol[iIdxCol].zSpan;
+            pIter->aIdxCol[iIdxCol++].nSpan = nSpan;
+            pIter->aIdxCol[iIdxCol].zSpan = &zSql[i+1];
           }else if( c=='"' || c=='\'' || c=='`' ){
             for(i++; 1; i++){
               if( zSql[i]==c ){
@@ -190869,11 +196079,19 @@ static char *rbuObjIterGetIndexWhere(sqlite3rbu *p, RbuObjIter *pIter){
             for(i++; 1; i++){
               if( zSql[i]==']' ) break;
             }
+          }else if( c=='-' && zSql[i+1]=='-' ){
+            for(i=i+2; zSql[i] && zSql[i]!='\n'; i++);
+            if( zSql[i]=='\0' ) break;
+          }else if( c=='/' && zSql[i+1]=='*' ){
+            for(i=i+2; zSql[i] && (zSql[i]!='*' || zSql[i+1]!='/'); i++);
+            if( zSql[i]=='\0' ) break;
+            i++;
           }
         }
         if( zSql[i] ){
           zRet = rbuStrndup(&zSql[i], &rc);
         }
+        pIter->nIdxCol = iIdxCol;
       }
     }
 
@@ -190918,11 +196136,11 @@ static int rbuObjIterPrepareAll(
       int nBind = 0;
 
       assert( pIter->eType!=RBU_PK_VTAB );
+      zPart = rbuObjIterGetIndexWhere(p, pIter);
       zCollist = rbuObjIterGetIndexCols(
           p, pIter, &zImposterCols, &zImposterPK, &zWhere, &nBind
       );
       zBind = rbuObjIterGetBindlist(p, nBind);
-      zPart = rbuObjIterGetIndexWhere(p, pIter);
 
       /* Create the imposter table used to write to this index. */
       sqlite3_test_control(SQLITE_TESTCTRL_IMPOSTER, p->dbMain, "main", 0, 1);
@@ -190954,12 +196172,24 @@ static int rbuObjIterPrepareAll(
       if( p->rc==SQLITE_OK ){
         char *zSql;
         if( rbuIsVacuum(p) ){
+          char *zStart = 0;
+          if( nOffset ){
+            zStart = rbuVacuumIndexStart(p, pIter);
+            if( zStart ){
+              sqlite3_free(zLimit);
+              zLimit = 0;
+            }
+          }
+
           zSql = sqlite3_mprintf(
-              "SELECT %s, 0 AS rbu_control FROM '%q' %s ORDER BY %s%s",
+              "SELECT %s, 0 AS rbu_control FROM '%q' %s %s %s ORDER BY %s%s",
               zCollist, 
               pIter->zDataTbl,
-              zPart, zCollist, zLimit
+              zPart, 
+              (zStart ? (zPart ? "AND" : "WHERE") : ""), zStart,
+              zCollist, zLimit
           );
+          sqlite3_free(zStart);
         }else
 
         if( pIter->eType==RBU_PK_EXTERNAL || pIter->eType==RBU_PK_NONE ){
@@ -190982,7 +196212,11 @@ static int rbuObjIterPrepareAll(
               zCollist, zLimit
           );
         }
-        p->rc = prepareFreeAndCollectError(p->dbRbu, &pIter->pSelect, pz, zSql);
+        if( p->rc==SQLITE_OK ){
+          p->rc = prepareFreeAndCollectError(p->dbRbu,&pIter->pSelect,pz,zSql);
+        }else{
+          sqlite3_free(zSql);
+        }
       }
 
       sqlite3_free(zImposterCols);
@@ -191082,18 +196316,42 @@ static int rbuObjIterPrepareAll(
       /* Create the SELECT statement to read keys from data_xxx */
       if( p->rc==SQLITE_OK ){
         const char *zRbuRowid = "";
+        char *zStart = 0;
+        char *zOrder = 0;
         if( bRbuRowid ){
           zRbuRowid = rbuIsVacuum(p) ? ",_rowid_ " : ",rbu_rowid";
         }
-        p->rc = prepareFreeAndCollectError(p->dbRbu, &pIter->pSelect, pz,
-            sqlite3_mprintf(
-              "SELECT %s,%s rbu_control%s FROM '%q'%s", 
-              zCollist, 
-              (rbuIsVacuum(p) ? "0 AS " : ""),
-              zRbuRowid,
-              pIter->zDataTbl, zLimit
-            )
-        );
+
+        if( rbuIsVacuum(p) ){
+          if( nOffset ){
+            zStart = rbuVacuumTableStart(p, pIter, bRbuRowid, zWrite);
+            if( zStart ){
+              sqlite3_free(zLimit);
+              zLimit = 0;
+            }
+          }
+          if( bRbuRowid ){
+            zOrder = rbuMPrintf(p, "_rowid_");
+          }else{
+            zOrder = rbuObjIterGetPkList(p, pIter, "", ", ", "");
+          }
+        }
+
+        if( p->rc==SQLITE_OK ){
+          p->rc = prepareFreeAndCollectError(p->dbRbu, &pIter->pSelect, pz,
+              sqlite3_mprintf(
+                "SELECT %s,%s rbu_control%s FROM '%q'%s %s %s %s",
+                zCollist, 
+                (rbuIsVacuum(p) ? "0 AS " : ""),
+                zRbuRowid,
+                pIter->zDataTbl, (zStart ? zStart : ""), 
+                (zOrder ? "ORDER BY" : ""), zOrder,
+                zLimit
+              )
+          );
+        }
+        sqlite3_free(zStart);
+        sqlite3_free(zOrder);
       }
 
       sqlite3_free(zWhere);
@@ -192408,10 +197666,11 @@ static void rbuIndexCntFunc(
   sqlite3_stmt *pStmt = 0;
   char *zErrmsg = 0;
   int rc;
+  sqlite3 *db = (rbuIsVacuum(p) ? p->dbRbu : p->dbMain);
 
   assert( nVal==1 );
   
-  rc = prepareFreeAndCollectError(p->dbMain, &pStmt, &zErrmsg, 
+  rc = prepareFreeAndCollectError(db, &pStmt, &zErrmsg, 
       sqlite3_mprintf("SELECT count(*) FROM sqlite_master "
         "WHERE type='index' AND tbl_name = %Q", sqlite3_value_text(apVal[0]))
   );
@@ -192426,7 +197685,7 @@ static void rbuIndexCntFunc(
     if( rc==SQLITE_OK ){
       sqlite3_result_int(pCtx, nIndex);
     }else{
-      sqlite3_result_error(pCtx, sqlite3_errmsg(p->dbMain), -1);
+      sqlite3_result_error(pCtx, sqlite3_errmsg(db), -1);
     }
   }
 
@@ -193320,9 +198579,7 @@ static int rbuVfsFileControl(sqlite3_file *pFile, int op, void *pArg){
       }else if( rc==SQLITE_NOTFOUND ){
         pRbu->pTargetFd = p;
         p->pRbu = pRbu;
-        if( p->openFlags & SQLITE_OPEN_MAIN_DB ){
-          rbuMainlistAdd(p);
-        }
+        rbuMainlistAdd(p);
         if( p->pWalFd ) p->pWalFd->pRbu = pRbu;
         rc = SQLITE_OK;
       }
@@ -193385,10 +198642,7 @@ static int rbuVfsShmLock(sqlite3_file *pFile, int ofst, int n, int flags){
     if( ofst==WAL_LOCK_CKPT && n==1 ) rc = SQLITE_BUSY;
   }else{
     int bCapture = 0;
-    if( n==1 && (flags & SQLITE_SHM_EXCLUSIVE)
-     && pRbu && pRbu->eStage==RBU_STAGE_CAPTURE
-     && (ofst==WAL_LOCK_WRITE || ofst==WAL_LOCK_CKPT || ofst==WAL_LOCK_READ0)
-    ){
+    if( pRbu && pRbu->eStage==RBU_STAGE_CAPTURE ){
       bCapture = 1;
     }
 
@@ -193421,20 +198675,24 @@ static int rbuVfsShmMap(
   ** rbu is in the RBU_STAGE_OAL state, use heap memory for *-shm space 
   ** instead of a file on disk.  */
   assert( p->openFlags & (SQLITE_OPEN_MAIN_DB|SQLITE_OPEN_TEMP_DB) );
-  if( eStage==RBU_STAGE_OAL || eStage==RBU_STAGE_MOVE ){
-    if( iRegion<=p->nShm ){
-      sqlite3_int64 nByte = (iRegion+1) * sizeof(char*);
-      char **apNew = (char**)sqlite3_realloc64(p->apShm, nByte);
-      if( apNew==0 ){
-        rc = SQLITE_NOMEM;
-      }else{
-        memset(&apNew[p->nShm], 0, sizeof(char*) * (1 + iRegion - p->nShm));
-        p->apShm = apNew;
-        p->nShm = iRegion+1;
-      }
+  if( eStage==RBU_STAGE_OAL ){
+    sqlite3_int64 nByte = (iRegion+1) * sizeof(char*);
+    char **apNew = (char**)sqlite3_realloc64(p->apShm, nByte);
+
+    /* This is an RBU connection that uses its own heap memory for the
+    ** pages of the *-shm file. Since no other process can have run
+    ** recovery, the connection must request *-shm pages in order
+    ** from start to finish.  */
+    assert( iRegion==p->nShm );
+    if( apNew==0 ){
+      rc = SQLITE_NOMEM;
+    }else{
+      memset(&apNew[p->nShm], 0, sizeof(char*) * (1 + iRegion - p->nShm));
+      p->apShm = apNew;
+      p->nShm = iRegion+1;
     }
 
-    if( rc==SQLITE_OK && p->apShm[iRegion]==0 ){
+    if( rc==SQLITE_OK ){
       char *pNew = (char*)sqlite3_malloc64(szRegion);
       if( pNew==0 ){
         rc = SQLITE_NOMEM;
@@ -193484,33 +198742,6 @@ static int rbuVfsShmUnmap(sqlite3_file *pFile, int delFlag){
   return rc;
 }
 
-/* 
-** A main database named zName has just been opened. The following 
-** function returns a pointer to a buffer owned by SQLite that contains
-** the name of the *-wal file this db connection will use. SQLite
-** happens to pass a pointer to this buffer when using xAccess()
-** or xOpen() to operate on the *-wal file.  
-*/
-static const char *rbuMainToWal(const char *zName, int flags){
-  int n = (int)strlen(zName);
-  const char *z = &zName[n];
-  if( flags & SQLITE_OPEN_URI ){
-    int odd = 0;
-    while( 1 ){
-      if( z[0]==0 ){
-        odd = 1 - odd;
-        if( odd && z[1]==0 ) break;
-      }
-      z++;
-    }
-    z += 2;
-  }else{
-    while( *z==0 ) z++;
-  }
-  z += (n + 8 + 1);
-  return z;
-}
-
 /*
 ** Open an rbu file handle.
 */
@@ -193559,7 +198790,7 @@ static int rbuVfsOpen(
       ** the name of the *-wal file this db connection will use. SQLite
       ** happens to pass a pointer to this buffer when using xAccess()
       ** or xOpen() to operate on the *-wal file.  */
-      pFd->zWal = rbuMainToWal(zName, flags);
+      pFd->zWal = sqlite3_filename_wal(zName);
     }
     else if( flags & SQLITE_OPEN_WAL ){
       rbu_file *pDb = rbuFindMaindb(pRbuVfs, zName, 0);
@@ -193574,7 +198805,7 @@ static int rbuVfsOpen(
           char *zCopy;
           if( rbuIsVacuum(pDb->pRbu) ){
             zBase = sqlite3_db_filename(pDb->pRbu->dbRbu, "main");
-            zBase = rbuMainToWal(zBase, SQLITE_OPEN_URI);
+            zBase = sqlite3_filename_wal(zBase);
           }
           nCopy = strlen(zBase);
           zCopy = sqlite3_malloc64(nCopy+2);
@@ -193663,7 +198894,8 @@ static int rbuVfsAccess(
   */
   if( rc==SQLITE_OK && flags==SQLITE_ACCESS_EXISTS ){
     rbu_file *pDb = rbuFindMaindb(pRbuVfs, zPath, 1);
-    if( pDb && pDb->pRbu && pDb->pRbu->eStage==RBU_STAGE_OAL ){
+    if( pDb && pDb->pRbu->eStage==RBU_STAGE_OAL ){
+      assert( pDb->pRbu );
       if( *pResOut ){
         rc = SQLITE_CANTOPEN;
       }else{
@@ -193893,7 +199125,7 @@ SQLITE_API sqlite3_int64 sqlite3rbu_temp_size(sqlite3rbu *pRbu){
 **
 ** This file contains an implementation of the "dbstat" virtual table.
 **
-** The dbstat virtual table is used to extract low-level formatting
+** The dbstat virtual table is used to extract low-level storage
 ** information from an SQLite database in order to implement the
 ** "sqlite3_analyzer" utility.  See the ../tool/spaceanal.tcl script
 ** for an example implementation.
@@ -193937,27 +199169,30 @@ SQLITE_API sqlite3_int64 sqlite3rbu_temp_size(sqlite3rbu *pRbu){
 **
 **      '/1c2/000/'               // Left-most child of 451st child of root
 */
-#define VTAB_SCHEMA                                                         \
-  "CREATE TABLE xx( "                                                       \
-  "  name       TEXT,             /* Name of table or index */"             \
-  "  path       TEXT,             /* Path to page from root */"             \
-  "  pageno     INTEGER,          /* Page number */"                        \
-  "  pagetype   TEXT,             /* 'internal', 'leaf' or 'overflow' */"   \
-  "  ncell      INTEGER,          /* Cells on page (0 for overflow) */"     \
-  "  payload    INTEGER,          /* Bytes of payload on this page */"      \
-  "  unused     INTEGER,          /* Bytes of unused space on this page */" \
-  "  mx_payload INTEGER,          /* Largest payload size of all cells */"  \
-  "  pgoffset   INTEGER,          /* Offset of page in file */"             \
-  "  pgsize     INTEGER,          /* Size of the page */"                   \
-  "  schema     TEXT HIDDEN       /* Database schema being analyzed */"     \
-  ");"
-
-
+static const char zDbstatSchema[] = 
+  "CREATE TABLE x("
+  " name       TEXT,"          /*  0 Name of table or index */
+  " path       TEXT,"          /*  1 Path to page from root (NULL for agg) */
+  " pageno     INTEGER,"       /*  2 Page number (page count for aggregates) */
+  " pagetype   TEXT,"          /*  3 'internal', 'leaf', 'overflow', or NULL */
+  " ncell      INTEGER,"       /*  4 Cells on page (0 for overflow) */
+  " payload    INTEGER,"       /*  5 Bytes of payload on this page */
+  " unused     INTEGER,"       /*  6 Bytes of unused space on this page */
+  " mx_payload INTEGER,"       /*  7 Largest payload size of all cells */
+  " pgoffset   INTEGER,"       /*  8 Offset of page in file (NULL for agg) */
+  " pgsize     INTEGER,"       /*  9 Size of the page (sum for aggregate) */
+  " schema     TEXT HIDDEN,"   /* 10 Database schema being analyzed */
+  " aggregate  BOOLEAN HIDDEN" /* 11 aggregate info for each table */
+  ")"
+;
+
+/* Forward reference to data structured used in this module */
 typedef struct StatTable StatTable;
 typedef struct StatCursor StatCursor;
 typedef struct StatPage StatPage;
 typedef struct StatCell StatCell;
 
+/* Size information for a single cell within a btree page */
 struct StatCell {
   int nLocal;                     /* Bytes of local payload */
   u32 iChildPg;                   /* Child node (or 0 if this is a leaf) */
@@ -193967,10 +199202,11 @@ struct StatCell {
   int iOvfl;                      /* Iterates through aOvfl[] */
 };
 
+/* Size information for a single btree page */
 struct StatPage {
-  u32 iPgno;
-  DbPage *pPg;
-  int iCell;
+  u32 iPgno;                      /* Page number */
+  DbPage *pPg;                    /* Page content */
+  int iCell;                      /* Current cell */
 
   char *zPath;                    /* Path to this page */
 
@@ -193980,34 +199216,38 @@ struct StatPage {
   int nUnused;                    /* Number of unused bytes on page */
   StatCell *aCell;                /* Array of parsed cells */
   u32 iRightChildPg;              /* Right-child page number (or 0) */
-  int nMxPayload;                 /* Largest payload of any cell on this page */
+  int nMxPayload;                 /* Largest payload of any cell on the page */
 };
 
+/* The cursor for scanning the dbstat virtual table */
 struct StatCursor {
-  sqlite3_vtab_cursor base;
+  sqlite3_vtab_cursor base;       /* base class.  MUST BE FIRST! */
   sqlite3_stmt *pStmt;            /* Iterates through set of root pages */
-  int isEof;                      /* After pStmt has returned SQLITE_DONE */
+  u8 isEof;                       /* After pStmt has returned SQLITE_DONE */
+  u8 isAgg;                       /* Aggregate results for each table */
   int iDb;                        /* Schema used for this query */
 
-  StatPage aPage[32];
+  StatPage aPage[32];             /* Pages in path to current page */
   int iPage;                      /* Current entry in aPage[] */
 
   /* Values to return. */
+  u32 iPageno;                    /* Value of 'pageno' column */
   char *zName;                    /* Value of 'name' column */
   char *zPath;                    /* Value of 'path' column */
-  u32 iPageno;                    /* Value of 'pageno' column */
   char *zPagetype;                /* Value of 'pagetype' column */
+  int nPage;                      /* Number of pages in current btree */
   int nCell;                      /* Value of 'ncell' column */
-  int nPayload;                   /* Value of 'payload' column */
-  int nUnused;                    /* Value of 'unused' column */
   int nMxPayload;                 /* Value of 'mx_payload' column */
+  i64 nUnused;                    /* Value of 'unused' column */
+  i64 nPayload;                   /* Value of 'payload' column */
   i64 iOffset;                    /* Value of 'pgOffset' column */
-  int szPage;                     /* Value of 'pgSize' column */
+  i64 szPage;                     /* Value of 'pgSize' column */
 };
 
+/* An instance of the DBSTAT virtual table */
 struct StatTable {
-  sqlite3_vtab base;
-  sqlite3 *db;
+  sqlite3_vtab base;              /* base class.  MUST BE FIRST! */
+  sqlite3 *db;                    /* Database connection that owns this vtab */
   int iDb;                        /* Index of database to analyze */
 };
 
@@ -194016,7 +199256,7 @@ struct StatTable {
 #endif
 
 /*
-** Connect to or create a statvfs virtual table.
+** Connect to or create a new DBSTAT virtual table.
 */
 static int statConnect(
   sqlite3 *db,
@@ -194040,7 +199280,8 @@ static int statConnect(
   }else{
     iDb = 0;
   }
-  rc = sqlite3_declare_vtab(db, VTAB_SCHEMA);
+  sqlite3_vtab_config(db, SQLITE_VTAB_DIRECTONLY);
+  rc = sqlite3_declare_vtab(db, zDbstatSchema);
   if( rc==SQLITE_OK ){
     pTab = (StatTable *)sqlite3_malloc64(sizeof(StatTable));
     if( pTab==0 ) rc = SQLITE_NOMEM_BKPT;
@@ -194058,7 +199299,7 @@ static int statConnect(
 }
 
 /*
-** Disconnect from or destroy a statvfs virtual table.
+** Disconnect from or destroy the DBSTAT virtual table.
 */
 static int statDisconnect(sqlite3_vtab *pVtab){
   sqlite3_free(pVtab);
@@ -194066,14 +199307,20 @@ static int statDisconnect(sqlite3_vtab *pVtab){
 }
 
 /*
-** There is no "best-index". This virtual table always does a linear
-** scan.  However, a schema=? constraint should cause this table to
-** operate on a different database schema, so check for it.
+** Compute the best query strategy and return the result in idxNum.
 **
-** idxNum is normally 0, but will be 1 if a schema=? constraint exists.
+**   idxNum-Bit        Meaning
+**   ----------        ----------------------------------------------
+**      0x01           There is a schema=? term in the WHERE clause
+**      0x02           There is a name=? term in the WHERE clause
+**      0x04           There is an aggregate=? term in the WHERE clause
+**      0x08           Output should be ordered by name and path
 */
 static int statBestIndex(sqlite3_vtab *tab, sqlite3_index_info *pIdxInfo){
   int i;
+  int iSchema = -1;
+  int iName = -1;
+  int iAgg = -1;
 
   /* Look for a valid schema=? constraint.  If found, change the idxNum to
   ** 1 and request the value of that constraint be sent to xFilter.  And
@@ -194081,16 +199328,40 @@ static int statBestIndex(sqlite3_vtab *tab, sqlite3_index_info *pIdxInfo){
   ** used.
   */
   for(i=0; i<pIdxInfo->nConstraint; i++){
-    if( pIdxInfo->aConstraint[i].iColumn!=10 ) continue;
-    if( pIdxInfo->aConstraint[i].usable==0 ) return SQLITE_CONSTRAINT;
     if( pIdxInfo->aConstraint[i].op!=SQLITE_INDEX_CONSTRAINT_EQ ) continue;
-    pIdxInfo->idxNum = 1;
-    pIdxInfo->estimatedCost = 1.0;
-    pIdxInfo->aConstraintUsage[i].argvIndex = 1;
-    pIdxInfo->aConstraintUsage[i].omit = 1;
-    break;
+    if( pIdxInfo->aConstraint[i].usable==0 ){
+      /* Force DBSTAT table should always be the right-most table in a join */
+      return SQLITE_CONSTRAINT;
+    }
+    switch( pIdxInfo->aConstraint[i].iColumn ){
+      case 0: {    /* name */
+        iName = i;
+        break;
+      }
+      case 10: {   /* schema */
+        iSchema = i;
+        break;
+      }
+      case 11: {   /* aggregate */
+        iAgg = i;
+        break;
+      }
+    }
   }
-
+  i = 0;
+  if( iSchema>=0 ){
+    pIdxInfo->aConstraintUsage[iSchema].argvIndex = ++i;
+    pIdxInfo->idxNum |= 0x01;
+  }
+  if( iName>=0 ){
+    pIdxInfo->aConstraintUsage[iName].argvIndex = ++i;
+    pIdxInfo->idxNum |= 0x02;
+  }
+  if( iAgg>=0 ){
+    pIdxInfo->aConstraintUsage[iAgg].argvIndex = ++i;
+    pIdxInfo->idxNum |= 0x04;
+  }
+  pIdxInfo->estimatedCost = 1.0;
 
   /* Records are always returned in ascending order of (name, path). 
   ** If this will satisfy the client, set the orderByConsumed flag so that 
@@ -194108,13 +199379,14 @@ static int statBestIndex(sqlite3_vtab *tab, sqlite3_index_info *pIdxInfo){
      )
   ){
     pIdxInfo->orderByConsumed = 1;
+    pIdxInfo->idxNum |= 0x08;
   }
 
   return SQLITE_OK;
 }
 
 /*
-** Open a new statvfs cursor.
+** Open a new DBSTAT cursor.
 */
 static int statOpen(sqlite3_vtab *pVTab, sqlite3_vtab_cursor **ppCursor){
   StatTable *pTab = (StatTable *)pVTab;
@@ -194164,8 +199436,18 @@ static void statResetCsr(StatCursor *pCsr){
   pCsr->isEof = 0;
 }
 
+/* Resize the space-used counters inside of the cursor */
+static void statResetCounts(StatCursor *pCsr){
+  pCsr->nCell = 0;
+  pCsr->nMxPayload = 0;
+  pCsr->nUnused = 0;
+  pCsr->nPayload = 0;
+  pCsr->szPage = 0;
+  pCsr->nPage = 0;
+}
+
 /*
-** Close a statvfs cursor.
+** Close a DBSTAT cursor.
 */
 static int statClose(sqlite3_vtab_cursor *pCursor){
   StatCursor *pCsr = (StatCursor *)pCursor;
@@ -194175,11 +199457,15 @@ static int statClose(sqlite3_vtab_cursor *pCursor){
   return SQLITE_OK;
 }
 
-static void getLocalPayload(
+/*
+** For a single cell on a btree page, compute the number of bytes of
+** content (payload) stored on that page.  That is to say, compute the
+** number of bytes of content not found on overflow pages.
+*/
+static int getLocalPayload(
   int nUsable,                    /* Usable bytes per page */
   u8 flags,                       /* Page flags */
-  int nTotal,                     /* Total record (payload) size */
-  int *pnLocal                    /* OUT: Bytes stored locally */
+  int nTotal                      /* Total record (payload) size */
 ){
   int nLocal;
   int nMinLocal;
@@ -194195,9 +199481,12 @@ static void getLocalPayload(
 
   nLocal = nMinLocal + (nTotal - nMinLocal) % (nUsable - 4);
   if( nLocal>nMaxLocal ) nLocal = nMinLocal;
-  *pnLocal = nLocal;
+  return nLocal;
 }
 
+/* Populate the StatPage object with information about the all
+** cells found on the page currently under analysis.
+*/
 static int statDecodePage(Btree *pBt, StatPage *p){
   int nUnused;
   int iOff;
@@ -194268,7 +199557,7 @@ static int statDecodePage(Btree *pBt, StatPage *p){
           iOff += sqlite3GetVarint(&aData[iOff], &dummy);
         }
         if( nPayload>(u32)p->nMxPayload ) p->nMxPayload = nPayload;
-        getLocalPayload(nUsable, p->flags, nPayload, &nLocal);
+        nLocal = getLocalPayload(nUsable, p->flags, nPayload);
         if( nLocal<0 ) goto statPageIsCorrupt;
         pCell->nLocal = nLocal;
         assert( nPayload>=(u32)nLocal );
@@ -194318,23 +199607,25 @@ static void statSizeAndOffset(StatCursor *pCsr){
   sqlite3_file *fd;
   sqlite3_int64 x[2];
 
-  /* The default page size and offset */
-  pCsr->szPage = sqlite3BtreeGetPageSize(pBt);
-  pCsr->iOffset = (i64)pCsr->szPage * (pCsr->iPageno - 1);
-
-  /* If connected to a ZIPVFS backend, override the page size and
-  ** offset with actual values obtained from ZIPVFS.
+  /* If connected to a ZIPVFS backend, find the page size and
+  ** offset from ZIPVFS.
   */
   fd = sqlite3PagerFile(pPager);
   x[0] = pCsr->iPageno;
   if( sqlite3OsFileControl(fd, 230440, &x)==SQLITE_OK ){
     pCsr->iOffset = x[0];
-    pCsr->szPage = (int)x[1];
+    pCsr->szPage += x[1];
+  }else{
+    /* Not ZIPVFS: The default page size and offset */
+    pCsr->szPage += sqlite3BtreeGetPageSize(pBt);
+    pCsr->iOffset = (i64)pCsr->szPage * (pCsr->iPageno - 1);
   }
 }
 
 /*
-** Move a statvfs cursor to the next entry in the file.
+** Move a DBSTAT cursor to the next entry.  Normally, the next
+** entry will be the next page, but in aggregated mode (pCsr->isAgg!=0),
+** the next entry is the next btree.
 */
 static int statNext(sqlite3_vtab_cursor *pCursor){
   int rc;
@@ -194350,6 +199641,8 @@ static int statNext(sqlite3_vtab_cursor *pCursor){
 
 statNextRestart:
   if( pCsr->aPage[0].pPg==0 ){
+    /* Start measuring space on the next btree */
+    statResetCounts(pCsr);
     rc = sqlite3_step(pCsr->pStmt);
     if( rc==SQLITE_ROW ){
       int nPage;
@@ -194362,44 +199655,47 @@ statNextRestart:
       rc = sqlite3PagerGet(pPager, iRoot, &pCsr->aPage[0].pPg, 0);
       pCsr->aPage[0].iPgno = iRoot;
       pCsr->aPage[0].iCell = 0;
-      pCsr->aPage[0].zPath = z = sqlite3_mprintf("/");
+      if( !pCsr->isAgg ){
+        pCsr->aPage[0].zPath = z = sqlite3_mprintf("/");
+        if( z==0 ) rc = SQLITE_NOMEM_BKPT;
+      }
       pCsr->iPage = 0;
-      if( z==0 ) rc = SQLITE_NOMEM_BKPT;
+      pCsr->nPage = 1;
     }else{
       pCsr->isEof = 1;
       return sqlite3_reset(pCsr->pStmt);
     }
   }else{
-
-    /* Page p itself has already been visited. */
+    /* Continue analyzing the btree previously started */
     StatPage *p = &pCsr->aPage[pCsr->iPage];
-
+    if( !pCsr->isAgg ) statResetCounts(pCsr);
     while( p->iCell<p->nCell ){
       StatCell *pCell = &p->aCell[p->iCell];
-      if( pCell->iOvfl<pCell->nOvfl ){
-        int nUsable;
+      while( pCell->iOvfl<pCell->nOvfl ){
+        int nUsable, iOvfl;
         sqlite3BtreeEnter(pBt);
         nUsable = sqlite3BtreeGetPageSize(pBt) - 
                         sqlite3BtreeGetReserveNoMutex(pBt);
         sqlite3BtreeLeave(pBt);
-        pCsr->zName = (char *)sqlite3_column_text(pCsr->pStmt, 0);
-        pCsr->iPageno = pCell->aOvfl[pCell->iOvfl];
-        pCsr->zPagetype = "overflow";
-        pCsr->nCell = 0;
-        pCsr->nMxPayload = 0;
-        pCsr->zPath = z = sqlite3_mprintf(
-            "%s%.3x+%.6x", p->zPath, p->iCell, pCell->iOvfl
-        );
+        pCsr->nPage++;
+        statSizeAndOffset(pCsr);
         if( pCell->iOvfl<pCell->nOvfl-1 ){
-          pCsr->nUnused = 0;
-          pCsr->nPayload = nUsable - 4;
+          pCsr->nPayload += nUsable - 4;
         }else{
-          pCsr->nPayload = pCell->nLastOvfl;
-          pCsr->nUnused = nUsable - 4 - pCsr->nPayload;
+          pCsr->nPayload += pCell->nLastOvfl;
+          pCsr->nUnused += nUsable - 4 - pCell->nLastOvfl;
         }
+        iOvfl = pCell->iOvfl;
         pCell->iOvfl++;
-        statSizeAndOffset(pCsr);
-        return z==0 ? SQLITE_NOMEM_BKPT : SQLITE_OK;
+        if( !pCsr->isAgg ){
+          pCsr->zName = (char *)sqlite3_column_text(pCsr->pStmt, 0);
+          pCsr->iPageno = pCell->aOvfl[iOvfl];
+          pCsr->zPagetype = "overflow";
+          pCsr->zPath = z = sqlite3_mprintf(
+              "%s%.3x+%.6x", p->zPath, p->iCell, iOvfl
+          );
+          return z==0 ? SQLITE_NOMEM_BKPT : SQLITE_OK;
+        }
       }
       if( p->iRightChildPg ) break;
       p->iCell++;
@@ -194407,8 +199703,13 @@ statNextRestart:
 
     if( !p->iRightChildPg || p->iCell>p->nCell ){
       statClearPage(p);
-      if( pCsr->iPage==0 ) return statNext(pCursor);
-      pCsr->iPage--;
+      if( pCsr->iPage>0 ){
+        pCsr->iPage--;
+      }else if( pCsr->isAgg ){
+        /* label-statNext-done:  When computing aggregate space usage over
+        ** an entire btree, this is the exit point from this function */
+        return SQLITE_OK;
+      }
       goto statNextRestart; /* Tail recursion */
     }
     pCsr->iPage++;
@@ -194424,10 +199725,13 @@ statNextRestart:
       p[1].iPgno = p->aCell[p->iCell].iChildPg;
     }
     rc = sqlite3PagerGet(pPager, p[1].iPgno, &p[1].pPg, 0);
+    pCsr->nPage++;
     p[1].iCell = 0;
-    p[1].zPath = z = sqlite3_mprintf("%s%.3x/", p->zPath, p->iCell);
+    if( !pCsr->isAgg ){
+      p[1].zPath = z = sqlite3_mprintf("%s%.3x/", p->zPath, p->iCell);
+      if( z==0 ) rc = SQLITE_NOMEM_BKPT;
+    }
     p->iCell++;
-    if( z==0 ) rc = SQLITE_NOMEM_BKPT;
   }
 
 
@@ -194457,16 +199761,23 @@ statNextRestart:
           pCsr->zPagetype = "corrupted";
           break;
       }
-      pCsr->nCell = p->nCell;
-      pCsr->nUnused = p->nUnused;
-      pCsr->nMxPayload = p->nMxPayload;
-      pCsr->zPath = z = sqlite3_mprintf("%s", p->zPath);
-      if( z==0 ) rc = SQLITE_NOMEM_BKPT;
+      pCsr->nCell += p->nCell;
+      pCsr->nUnused += p->nUnused;
+      if( p->nMxPayload>pCsr->nMxPayload ) pCsr->nMxPayload = p->nMxPayload;
+      if( !pCsr->isAgg ){
+        pCsr->zPath = z = sqlite3_mprintf("%s", p->zPath);
+        if( z==0 ) rc = SQLITE_NOMEM_BKPT;
+      }
       nPayload = 0;
       for(i=0; i<p->nCell; i++){
         nPayload += p->aCell[i].nLocal;
       }
-      pCsr->nPayload = nPayload;
+      pCsr->nPayload += nPayload;
+
+      /* If computing aggregate space usage by btree, continue with the
+      ** next page.  The loop will exit via the return at label-statNext-done
+      */
+      if( pCsr->isAgg ) goto statNextRestart;
     }
   }
 
@@ -194478,6 +199789,10 @@ static int statEof(sqlite3_vtab_cursor *pCursor){
   return pCsr->isEof;
 }
 
+/* Initialize a cursor according to the query plan idxNum using the
+** arguments in argv[0].  See statBestIndex() for a description of the
+** meaning of the bits in idxNum.
+*/
 static int statFilter(
   sqlite3_vtab_cursor *pCursor, 
   int idxNum, const char *idxStr,
@@ -194485,29 +199800,52 @@ static int statFilter(
 ){
   StatCursor *pCsr = (StatCursor *)pCursor;
   StatTable *pTab = (StatTable*)(pCursor->pVtab);
-  char *zSql;
-  int rc = SQLITE_OK;
+  sqlite3_str *pSql;      /* Query of btrees to analyze */
+  char *zSql;             /* String value of pSql */
+  int iArg = 0;           /* Count of argv[] parameters used so far */
+  int rc = SQLITE_OK;     /* Result of this operation */
+  const char *zName = 0;  /* Only provide analysis of this table */
 
-  if( idxNum==1 ){
-    const char *zDbase = (const char*)sqlite3_value_text(argv[0]);
+  statResetCsr(pCsr);
+  sqlite3_finalize(pCsr->pStmt);
+  pCsr->pStmt = 0;
+  if( idxNum & 0x01 ){
+    /* schema=? constraint is present.  Get its value */
+    const char *zDbase = (const char*)sqlite3_value_text(argv[iArg++]);
     pCsr->iDb = sqlite3FindDbName(pTab->db, zDbase);
     if( pCsr->iDb<0 ){
-      sqlite3_free(pCursor->pVtab->zErrMsg);
-      pCursor->pVtab->zErrMsg = sqlite3_mprintf("no such schema: %s", zDbase);
-      return pCursor->pVtab->zErrMsg ? SQLITE_ERROR : SQLITE_NOMEM_BKPT;
+      pCsr->iDb = 0;
+      pCsr->isEof = 1;
+      return SQLITE_OK;
     }
   }else{
     pCsr->iDb = pTab->iDb;
   }
-  statResetCsr(pCsr);
-  sqlite3_finalize(pCsr->pStmt);
-  pCsr->pStmt = 0;
-  zSql = sqlite3_mprintf(
-      "SELECT 'sqlite_master' AS name, 1 AS rootpage, 'table' AS type"
-      "  UNION ALL  "
-      "SELECT name, rootpage, type"
-      "  FROM \"%w\".sqlite_master WHERE rootpage!=0"
-      "  ORDER BY name", pTab->db->aDb[pCsr->iDb].zDbSName);
+  if( idxNum & 0x02 ){
+    /* name=? constraint is present */
+    zName = (const char*)sqlite3_value_text(argv[iArg++]);
+  }
+  if( idxNum & 0x04 ){
+    /* aggregate=? constraint is present */
+    pCsr->isAgg = sqlite3_value_double(argv[iArg++])!=0.0;
+  }else{
+    pCsr->isAgg = 0;
+  }
+  pSql = sqlite3_str_new(pTab->db);
+  sqlite3_str_appendf(pSql,
+      "SELECT * FROM ("
+        "SELECT 'sqlite_master' AS name,1 AS rootpage,'table' AS type"
+        " UNION ALL "
+        "SELECT name,rootpage,type"
+        " FROM \"%w\".sqlite_master WHERE rootpage!=0)",
+      pTab->db->aDb[pCsr->iDb].zDbSName);
+  if( zName ){
+    sqlite3_str_appendf(pSql, "WHERE name=%Q", zName);
+  }
+  if( idxNum & 0x08 ){
+    sqlite3_str_appendf(pSql, " ORDER BY name");
+  }
+  zSql = sqlite3_str_finish(pSql);
   if( zSql==0 ){
     return SQLITE_NOMEM_BKPT;
   }else{
@@ -194532,13 +199870,21 @@ static int statColumn(
       sqlite3_result_text(ctx, pCsr->zName, -1, SQLITE_TRANSIENT);
       break;
     case 1:            /* path */
-      sqlite3_result_text(ctx, pCsr->zPath, -1, SQLITE_TRANSIENT);
+      if( !pCsr->isAgg ){
+        sqlite3_result_text(ctx, pCsr->zPath, -1, SQLITE_TRANSIENT);
+      }
       break;
     case 2:            /* pageno */
-      sqlite3_result_int64(ctx, pCsr->iPageno);
+      if( pCsr->isAgg ){
+        sqlite3_result_int64(ctx, pCsr->nPage);
+      }else{
+        sqlite3_result_int64(ctx, pCsr->iPageno);
+      }
       break;
     case 3:            /* pagetype */
-      sqlite3_result_text(ctx, pCsr->zPagetype, -1, SQLITE_STATIC);
+      if( !pCsr->isAgg ){
+        sqlite3_result_text(ctx, pCsr->zPagetype, -1, SQLITE_STATIC);
+      }
       break;
     case 4:            /* ncell */
       sqlite3_result_int(ctx, pCsr->nCell);
@@ -194553,17 +199899,23 @@ static int statColumn(
       sqlite3_result_int(ctx, pCsr->nMxPayload);
       break;
     case 8:            /* pgoffset */
-      sqlite3_result_int64(ctx, pCsr->iOffset);
+      if( !pCsr->isAgg ){
+        sqlite3_result_int64(ctx, pCsr->iOffset);
+      }
       break;
     case 9:            /* pgsize */
       sqlite3_result_int(ctx, pCsr->szPage);
       break;
-    default: {          /* schema */
+    case 10: {         /* schema */
       sqlite3 *db = sqlite3_context_db_handle(ctx);
       int iDb = pCsr->iDb;
       sqlite3_result_text(ctx, db->aDb[iDb].zDbSName, -1, SQLITE_STATIC);
       break;
     }
+    default: {         /* aggregate */
+      sqlite3_result_int(ctx, pCsr->isAgg);
+      break;
+    }
   }
   return SQLITE_OK;
 }
@@ -194687,6 +200039,7 @@ static int dbpageConnect(
   DbpageTable *pTab = 0;
   int rc = SQLITE_OK;
 
+  sqlite3_vtab_config(db, SQLITE_VTAB_DIRECTONLY);
   rc = sqlite3_declare_vtab(db, 
           "CREATE TABLE x(pgno INTEGER PRIMARY KEY, data BLOB, schema HIDDEN)");
   if( rc==SQLITE_OK ){
@@ -196656,7 +202009,9 @@ SQLITE_API int sqlite3session_diff(
       }
       sqlite3_free((char*)azCol);
       if( bMismatch ){
-        *pzErrMsg = sqlite3_mprintf("table schemas do not match");
+        if( pzErrMsg ){
+          *pzErrMsg = sqlite3_mprintf("table schemas do not match");
+        }
         rc = SQLITE_SCHEMA;
       }
       if( bHasPk==0 ){
@@ -196862,12 +202217,12 @@ SQLITE_API int sqlite3session_attach(
 ** set *pRc to SQLITE_NOMEM and return non-zero.
 */
 static int sessionBufferGrow(SessionBuffer *p, size_t nByte, int *pRc){
-  if( *pRc==SQLITE_OK && p->nAlloc-p->nBuf<nByte ){
+  if( *pRc==SQLITE_OK && (size_t)(p->nAlloc-p->nBuf)<nByte ){
     u8 *aNew;
     i64 nNew = p->nAlloc ? p->nAlloc : 128;
     do {
       nNew = nNew*2;
-    }while( (nNew-p->nBuf)<nByte );
+    }while( (size_t)(nNew-p->nBuf)<nByte );
 
     aNew = (u8 *)sqlite3_realloc64(p->aBuf, nNew);
     if( 0==aNew ){
@@ -200673,7 +206028,7 @@ struct Fts5PhraseIter {
 **
 ** xSetAuxdata(pFts5, pAux, xDelete)
 **
-**   Save the pointer passed as the second argument as the extension functions 
+**   Save the pointer passed as the second argument as the extension function's 
 **   "auxiliary data". The pointer may then be retrieved by the current or any
 **   future invocation of the same fts5 extension function made as part of
 **   the same MATCH query using the xGetAuxdata() API.
@@ -200915,8 +206270,8 @@ struct Fts5ExtensionApi {
 **
 **   There are several ways to approach this in FTS5:
 **
-**   <ol><li> By mapping all synonyms to a single token. In this case, the 
-**            In the above example, this means that the tokenizer returns the
+**   <ol><li> By mapping all synonyms to a single token. In this case, using
+**            the above example, this means that the tokenizer returns the
 **            same token for inputs "first" and "1st". Say that token is in
 **            fact "first", so that when the user inserts the document "I won
 **            1st place" entries are added to the index for tokens "i", "won",
@@ -201151,6 +206506,11 @@ typedef sqlite3_uint64 u64;
 */
 #define FTS5_MAX_PREFIX_INDEXES 31
 
+/*
+** Maximum segments permitted in a single index 
+*/
+#define FTS5_MAX_SEGMENT 2000
+
 #define FTS5_DEFAULT_NEARDIST 10
 #define FTS5_DEFAULT_RANK     "bm25"
 
@@ -201268,6 +206628,7 @@ struct Fts5Config {
   char *zContentExprlist;
   Fts5Tokenizer *pTok;
   fts5_tokenizer *pTokApi;
+  int bLock;                      /* True when table is preparing statement */
 
   /* Values loaded from the %_config table */
   int iCookie;                    /* Incremented when %_config is modified */
@@ -201506,6 +206867,11 @@ static int sqlite3Fts5IterNextFrom(Fts5IndexIter*, i64 iMatch);
 */
 static void sqlite3Fts5IterClose(Fts5IndexIter*);
 
+/*
+** Close the reader blob handle, if it is open.
+*/
+static void sqlite3Fts5IndexCloseReader(Fts5Index*);
+
 /*
 ** This interface is used by the fts5vocab module.
 */
@@ -201784,6 +207150,7 @@ static int sqlite3Fts5ExprEof(Fts5Expr*);
 static i64 sqlite3Fts5ExprRowid(Fts5Expr*);
 
 static void sqlite3Fts5ExprFree(Fts5Expr*);
+static int sqlite3Fts5ExprAnd(Fts5Expr **pp1, Fts5Expr *p2);
 
 /* Called during startup to register a UDF with SQLite */
 static int sqlite3Fts5ExprInit(Fts5Global*, sqlite3*);
@@ -202065,6 +207432,7 @@ typedef union {
 #define sqlite3Fts5ParserCTX_STORE
 #define fts5YYNSTATE             35
 #define fts5YYNRULE              28
+#define fts5YYNRULE_WITH_ACTION  28
 #define fts5YYNFTS5TOKEN             16
 #define fts5YY_MAX_SHIFT         34
 #define fts5YY_MIN_SHIFTREDUCE   52
@@ -202635,15 +208003,18 @@ static fts5YYACTIONTYPE fts5yy_find_shift_action(
   do{
     i = fts5yy_shift_ofst[stateno];
     assert( i>=0 );
-    /* assert( i+fts5YYNFTS5TOKEN<=(int)fts5YY_NLOOKAHEAD ); */
+    assert( i<=fts5YY_ACTTAB_COUNT );
+    assert( i+fts5YYNFTS5TOKEN<=(int)fts5YY_NLOOKAHEAD );
     assert( iLookAhead!=fts5YYNOCODE );
     assert( iLookAhead < fts5YYNFTS5TOKEN );
     i += iLookAhead;
-    if( i>=fts5YY_NLOOKAHEAD || fts5yy_lookahead[i]!=iLookAhead ){
+    assert( i<(int)fts5YY_NLOOKAHEAD );
+    if( fts5yy_lookahead[i]!=iLookAhead ){
 #ifdef fts5YYFALLBACK
       fts5YYCODETYPE iFallback;            /* Fallback token */
-      if( iLookAhead<sizeof(fts5yyFallback)/sizeof(fts5yyFallback[0])
-             && (iFallback = fts5yyFallback[iLookAhead])!=0 ){
+      assert( iLookAhead<sizeof(fts5yyFallback)/sizeof(fts5yyFallback[0]) );
+      iFallback = fts5yyFallback[iLookAhead];
+      if( iFallback!=0 ){
 #ifndef NDEBUG
         if( fts5yyTraceFILE ){
           fprintf(fts5yyTraceFILE, "%sFALLBACK %s => %s\n",
@@ -202658,16 +208029,8 @@ static fts5YYACTIONTYPE fts5yy_find_shift_action(
 #ifdef fts5YYWILDCARD
       {
         int j = i - iLookAhead + fts5YYWILDCARD;
-        if( 
-#if fts5YY_SHIFT_MIN+fts5YYWILDCARD<0
-          j>=0 &&
-#endif
-#if fts5YY_SHIFT_MAX+fts5YYWILDCARD>=fts5YY_ACTTAB_COUNT
-          j<fts5YY_ACTTAB_COUNT &&
-#endif
-          j<(int)(sizeof(fts5yy_lookahead)/sizeof(fts5yy_lookahead[0])) &&
-          fts5yy_lookahead[j]==fts5YYWILDCARD && iLookAhead>0
-        ){
+        assert( j<(int)(sizeof(fts5yy_lookahead)/sizeof(fts5yy_lookahead[0])) );
+        if( fts5yy_lookahead[j]==fts5YYWILDCARD && iLookAhead>0 ){
 #ifndef NDEBUG
           if( fts5yyTraceFILE ){
             fprintf(fts5yyTraceFILE, "%sWILDCARD %s => %s\n",
@@ -202681,6 +208044,7 @@ static fts5YYACTIONTYPE fts5yy_find_shift_action(
 #endif /* fts5YYWILDCARD */
       return fts5yy_default[stateno];
     }else{
+      assert( i>=0 && i<sizeof(fts5yy_action)/sizeof(fts5yy_action[0]) );
       return fts5yy_action[i];
     }
   }while(1);
@@ -202898,12 +208262,15 @@ static fts5YYACTIONTYPE fts5yy_reduce(
   if( fts5yyTraceFILE && fts5yyruleno<(int)(sizeof(fts5yyRuleName)/sizeof(fts5yyRuleName[0])) ){
     fts5yysize = fts5yyRuleInfoNRhs[fts5yyruleno];
     if( fts5yysize ){
-      fprintf(fts5yyTraceFILE, "%sReduce %d [%s], go to state %d.\n",
+      fprintf(fts5yyTraceFILE, "%sReduce %d [%s]%s, pop back to state %d.\n",
         fts5yyTracePrompt,
-        fts5yyruleno, fts5yyRuleName[fts5yyruleno], fts5yymsp[fts5yysize].stateno);
+        fts5yyruleno, fts5yyRuleName[fts5yyruleno],
+        fts5yyruleno<fts5YYNRULE_WITH_ACTION ? "" : " without external action",
+        fts5yymsp[fts5yysize].stateno);
     }else{
-      fprintf(fts5yyTraceFILE, "%sReduce %d [%s].\n",
-        fts5yyTracePrompt, fts5yyruleno, fts5yyRuleName[fts5yyruleno]);
+      fprintf(fts5yyTraceFILE, "%sReduce %d [%s]%s.\n",
+        fts5yyTracePrompt, fts5yyruleno, fts5yyRuleName[fts5yyruleno],
+        fts5yyruleno<fts5YYNRULE_WITH_ACTION ? "" : " without external action");
     }
   }
 #endif /* NDEBUG */
@@ -203381,13 +208748,12 @@ static void sqlite3Fts5Parser(
 */
 static int sqlite3Fts5ParserFallback(int iToken){
 #ifdef fts5YYFALLBACK
-  if( iToken<(int)(sizeof(fts5yyFallback)/sizeof(fts5yyFallback[0])) ){
-    return fts5yyFallback[iToken];
-  }
+  assert( iToken<(int)(sizeof(fts5yyFallback)/sizeof(fts5yyFallback[0])) );
+  return fts5yyFallback[iToken];
 #else
   (void)iToken;
-#endif
   return 0;
+#endif
 }
 
 /*
@@ -204287,10 +209653,19 @@ static int sqlite3Fts5PoslistNext64(
     i64 iOff = *piOff;
     int iVal;
     fts5FastGetVarint32(a, i, iVal);
-    if( iVal==1 ){
+    if( iVal<=1 ){
+      if( iVal==0 ){
+        *pi = i;
+        return 0;
+      }
       fts5FastGetVarint32(a, i, iVal);
       iOff = ((i64)iVal) << 32;
       fts5FastGetVarint32(a, i, iVal);
+      if( iVal<2 ){
+        /* This is a corrupt record. So stop parsing it here. */
+        *piOff = -1;
+        return 1;
+      }
     }
     *piOff = iOff + ((iVal-2) & 0x7FFFFFFF);
     *pi = i;
@@ -204527,7 +209902,7 @@ static void sqlite3Fts5TermsetFree(Fts5Termset *p){
 #define FTS5_DEFAULT_HASHSIZE    (1024*1024)
 
 /* Maximum allowed page size */
-#define FTS5_MAX_PAGE_SIZE (128*1024)
+#define FTS5_MAX_PAGE_SIZE (64*1024)
 
 static int fts5_iswhitespace(char x){
   return (x==' ');
@@ -204654,7 +210029,7 @@ static int fts5Dequote(char *z){
   assert( q=='[' || q=='\'' || q=='"' || q=='`' );
   if( q=='[' ) q = ']';  
 
-  while( ALWAYS(z[iIn]) ){
+  while( z[iIn] ){
     if( z[iIn]==q ){
       if( z[iIn+1]!=q ){
         /* Character iIn was the close quote. */
@@ -205187,7 +210562,7 @@ static int sqlite3Fts5ConfigDeclareVtab(Fts5Config *pConfig){
     rc = sqlite3_declare_vtab(pConfig->db, zSql);
     sqlite3_free(zSql);
   }
-  
+ 
   return rc;
 }
 
@@ -205332,7 +210707,7 @@ static int sqlite3Fts5ConfigSetValue(
     if( SQLITE_INTEGER==sqlite3_value_numeric_type(pVal) ){
       pgsz = sqlite3_value_int(pVal);
     }
-    if( pgsz<=0 || pgsz>FTS5_MAX_PAGE_SIZE ){
+    if( pgsz<32 || pgsz>FTS5_MAX_PAGE_SIZE ){
       *pbBadkey = 1;
     }else{
       pConfig->pgsz = pgsz;
@@ -205385,6 +210760,7 @@ static int sqlite3Fts5ConfigSetValue(
       *pbBadkey = 1;
     }else{
       if( nCrisisMerge<=1 ) nCrisisMerge = FTS5_DEFAULT_CRISISMERGE;
+      if( nCrisisMerge>=FTS5_MAX_SEGMENT ) nCrisisMerge = FTS5_MAX_SEGMENT-1;
       pConfig->nCrisisMerge = nCrisisMerge;
     }
   }
@@ -205775,6 +211151,42 @@ static void sqlite3Fts5ExprFree(Fts5Expr *p){
   }
 }
 
+static int sqlite3Fts5ExprAnd(Fts5Expr **pp1, Fts5Expr *p2){
+  Fts5Parse sParse;
+  memset(&sParse, 0, sizeof(sParse));
+
+  if( *pp1 ){
+    Fts5Expr *p1 = *pp1;
+    int nPhrase = p1->nPhrase + p2->nPhrase;
+
+    p1->pRoot = sqlite3Fts5ParseNode(&sParse, FTS5_AND, p1->pRoot, p2->pRoot,0);
+    p2->pRoot = 0;
+
+    if( sParse.rc==SQLITE_OK ){
+      Fts5ExprPhrase **ap = (Fts5ExprPhrase**)sqlite3_realloc(
+          p1->apExprPhrase, nPhrase * sizeof(Fts5ExprPhrase*)
+      );
+      if( ap==0 ){
+        sParse.rc = SQLITE_NOMEM;
+      }else{
+        int i;
+        memmove(&ap[p2->nPhrase], ap, p1->nPhrase*sizeof(Fts5ExprPhrase*));
+        for(i=0; i<p2->nPhrase; i++){
+          ap[i] = p2->apExprPhrase[i];
+        }
+        p1->nPhrase = nPhrase;
+        p1->apExprPhrase = ap;
+      }
+    }
+    sqlite3_free(p2->apExprPhrase);
+    sqlite3_free(p2);
+  }else{
+    *pp1 = p2;
+  }
+
+  return sParse.rc;
+}
+
 /*
 ** Argument pTerm must be a synonym iterator. Return the current rowid
 ** that it points to.
@@ -207946,10 +213358,12 @@ static void fts5ExprFunction(
   azConfig[1] = "main";
   azConfig[2] = "tbl";
   for(i=3; iArg<nArg; iArg++){
-    azConfig[i++] = (const char*)sqlite3_value_text(apVal[iArg]);
+    const char *z = (const char*)sqlite3_value_text(apVal[iArg]);
+    azConfig[i++] = (z ? z : "");
   }
 
   zExpr = (const char*)sqlite3_value_text(apVal[0]);
+  if( zExpr==0 ) zExpr = "";
 
   rc = sqlite3Fts5ConfigParse(pGlobal, db, nConfig, azConfig, &pConfig, &zErr);
   if( rc==SQLITE_OK ){
@@ -209120,11 +214534,6 @@ static void sqlite3Fts5HashScanEntry(
 #define FTS5_SEGMENT_ROWID(segid, pgno)       fts5_dri(segid, 0, 0, pgno)
 #define FTS5_DLIDX_ROWID(segid, height, pgno) fts5_dri(segid, 1, height, pgno)
 
-/*
-** Maximum segments permitted in a single index 
-*/
-#define FTS5_MAX_SEGMENT 2000
-
 #ifdef SQLITE_DEBUG
 static int sqlite3Fts5Corrupt() { return SQLITE_CORRUPT_VTAB; }
 #endif
@@ -209500,7 +214909,7 @@ static int fts5LeafFirstTermOff(Fts5Data *pLeaf){
 /*
 ** Close the read-only blob handle, if it is open.
 */
-static void fts5CloseReader(Fts5Index *p){
+static void sqlite3Fts5IndexCloseReader(Fts5Index *p){
   if( p->pReader ){
     sqlite3_blob *pReader = p->pReader;
     p->pReader = 0;
@@ -209529,7 +214938,7 @@ static Fts5Data *fts5DataRead(Fts5Index *p, i64 iRowid){
       assert( p->pReader==0 );
       p->pReader = pBlob;
       if( rc!=SQLITE_OK ){
-        fts5CloseReader(p);
+        sqlite3Fts5IndexCloseReader(p);
       }
       if( rc==SQLITE_ABORT ) rc = SQLITE_OK;
     }
@@ -209571,6 +214980,7 @@ static Fts5Data *fts5DataRead(Fts5Index *p, i64 iRowid){
       }else{
         /* TODO1: Fix this */
         pRet->p[nByte] = 0x00;
+        pRet->p[nByte+1] = 0x00;
         pRet->szLeaf = fts5GetU16(&pRet->p[2]);
       }
     }
@@ -209593,7 +215003,7 @@ static void fts5DataRelease(Fts5Data *pData){
 static Fts5Data *fts5LeafRead(Fts5Index *p, i64 iRowid){
   Fts5Data *pRet = fts5DataRead(p, iRowid);
   if( pRet ){
-    if( pRet->szLeaf>pRet->nn ){
+    if( pRet->nn<4 || pRet->szLeaf>pRet->nn ){
       p->rc = FTS5_CORRUPT;
       fts5DataRelease(pRet);
       pRet = 0;
@@ -209873,7 +215283,7 @@ static Fts5Structure *fts5StructureReadUncached(Fts5Index *p){
     /* TODO: Do we need this if the leaf-index is appended? Probably... */
     memset(&pData->p[pData->nn], 0, FTS5_DATA_PADDING);
     p->rc = fts5StructureDecode(pData->p, pData->nn, &iCookie, &pRet);
-    if( p->rc==SQLITE_OK && pConfig->iCookie!=iCookie ){
+    if( p->rc==SQLITE_OK && (pConfig->pgsz==0 || pConfig->iCookie!=iCookie) ){
       p->rc = sqlite3Fts5ConfigLoad(pConfig, iCookie);
     }
     fts5DataRelease(pData);
@@ -213834,8 +219244,14 @@ static void fts5MergePrefixLists(
     ** first rowid in one input is a large negative number, and the first in
     ** the other a non-negative number, the delta for the non-negative
     ** number will be larger on disk than the literal integer value
-    ** was.  */
-    if( sqlite3Fts5BufferSize(&p->rc, &out, p1->n + p2->n + 9) ) return;
+    ** was.  
+    **
+    ** Or, if the input position-lists are corrupt, then the output might
+    ** include up to 2 extra 10-byte positions created by interpreting -1
+    ** (the value PoslistNext64() uses for EOF) as a position and appending
+    ** it to the output. This can happen at most once for each input 
+    ** position-list, hence two 10 byte paddings.  */
+    if( sqlite3Fts5BufferSize(&p->rc, &out, p1->n + p2->n + 9+10+10) ) return;
     fts5DoclistIterInit(p1, &i1);
     fts5DoclistIterInit(p2, &i2);
 
@@ -213846,6 +219262,7 @@ static void fts5MergePrefixLists(
         fts5BufferSafeAppendBlob(&out, i1.aPoslist, i1.nPoslist+i1.nSize);
         fts5DoclistIterNext(&i1);
         if( i1.aPoslist==0 ) break;
+        assert( out.n<=((i1.aPoslist-p1->p) + (i2.aPoslist-p2->p)+9+10+10) );
       }
       else if( i2.iRowid!=i1.iRowid ){
         /* Copy entry from i2 */
@@ -213853,6 +219270,7 @@ static void fts5MergePrefixLists(
         fts5BufferSafeAppendBlob(&out, i2.aPoslist, i2.nPoslist+i2.nSize);
         fts5DoclistIterNext(&i2);
         if( i2.aPoslist==0 ) break;
+        assert( out.n<=((i1.aPoslist-p1->p) + (i2.aPoslist-p2->p)+9+10+10) );
       }
       else{
         /* Merge the two position lists. */ 
@@ -213869,14 +219287,17 @@ static void fts5MergePrefixLists(
         Fts5PoslistWriter writer;
         memset(&writer, 0, sizeof(writer));
 
+        /* See the earlier comment in this function for an explanation of why
+        ** corrupt input position lists might cause the output to consume
+        ** at most 20 bytes of unexpected space. */
         fts5MergeAppendDocid(&out, iLastRowid, i2.iRowid);
         fts5BufferZero(&tmp);
-        sqlite3Fts5BufferSize(&p->rc, &tmp, i1.nPoslist + i2.nPoslist);
+        sqlite3Fts5BufferSize(&p->rc, &tmp, i1.nPoslist + i2.nPoslist + 10 + 10);
         if( p->rc ) break;
 
         sqlite3Fts5PoslistNext64(a1, i1.nPoslist, &iOff1, &iPos1);
         sqlite3Fts5PoslistNext64(a2, i2.nPoslist, &iOff2, &iPos2);
-        assert( iPos1>=0 && iPos2>=0 );
+        assert_nc( iPos1>=0 && iPos2>=0 );
 
         if( iPos1<iPos2 ){
           sqlite3Fts5PoslistSafeAppend(&tmp, &iPrev, iPos1);
@@ -213885,7 +219306,6 @@ static void fts5MergePrefixLists(
           sqlite3Fts5PoslistSafeAppend(&tmp, &iPrev, iPos2);
           sqlite3Fts5PoslistNext64(a2, i2.nPoslist, &iOff2, &iPos2);
         }
-
         if( iPos1>=0 && iPos2>=0 ){
           while( 1 ){
             if( iPos1<iPos2 ){
@@ -213910,7 +219330,7 @@ static void fts5MergePrefixLists(
           aCopy = &a1[iOff1];
           nCopy = i1.nPoslist - iOff1;
         }else{
-          assert( iPos2>=0 && iPos2!=iPrev );
+          assert_nc( iPos2>=0 && iPos2!=iPrev );
           sqlite3Fts5PoslistSafeAppend(&tmp, &iPrev, iPos2);
           aCopy = &a2[iOff2];
           nCopy = i2.nPoslist - iOff2;
@@ -213920,12 +219340,19 @@ static void fts5MergePrefixLists(
         }
 
         /* WRITEPOSLISTSIZE */
+        assert_nc( tmp.n<=i1.nPoslist+i2.nPoslist );
+        assert( tmp.n<=i1.nPoslist+i2.nPoslist+10+10 );
+        if( tmp.n>i1.nPoslist+i2.nPoslist ){
+          if( p->rc==SQLITE_OK ) p->rc = FTS5_CORRUPT;
+          break;
+        }
         fts5BufferSafeAppendVarint(&out, tmp.n * 2);
         fts5BufferSafeAppendBlob(&out, tmp.p, tmp.n);
         fts5DoclistIterNext(&i1);
         fts5DoclistIterNext(&i2);
-        assert( out.n<=(p1->n+p2->n+9) );
+        assert_nc( out.n<=(p1->n+p2->n+9) );
         if( i1.aPoslist==0 || i2.aPoslist==0 ) break;
+        assert( out.n<=((i1.aPoslist-p1->p) + (i2.aPoslist-p2->p)+9+10+10) );
       }
     }
 
@@ -213937,7 +219364,7 @@ static void fts5MergePrefixLists(
       fts5MergeAppendDocid(&out, iLastRowid, i2.iRowid);
       fts5BufferSafeAppendBlob(&out, i2.aPoslist, i2.aEof - i2.aPoslist);
     }
-    assert( out.n<=(p1->n+p2->n+9) );
+    assert_nc( out.n<=(p1->n+p2->n+9) );
 
     fts5BufferSet(&p->rc, p1, out.n, out.p);
     fts5BufferFree(&tmp);
@@ -214072,7 +219499,7 @@ static int sqlite3Fts5IndexBeginWrite(Fts5Index *p, int bDelete, i64 iRowid){
 static int sqlite3Fts5IndexSync(Fts5Index *p){
   assert( p->rc==SQLITE_OK );
   fts5IndexFlush(p);
-  fts5CloseReader(p);
+  sqlite3Fts5IndexCloseReader(p);
   return fts5IndexReturn(p);
 }
 
@@ -214083,7 +219510,7 @@ static int sqlite3Fts5IndexSync(Fts5Index *p){
 ** records must be invalidated.
 */
 static int sqlite3Fts5IndexRollback(Fts5Index *p){
-  fts5CloseReader(p);
+  sqlite3Fts5IndexCloseReader(p);
   fts5IndexDiscardData(p);
   fts5StructureInvalidate(p);
   /* assert( p->rc==SQLITE_OK ); */
@@ -214098,6 +219525,7 @@ static int sqlite3Fts5IndexRollback(Fts5Index *p){
 static int sqlite3Fts5IndexReinit(Fts5Index *p){
   Fts5Structure s;
   fts5StructureInvalidate(p);
+  fts5IndexDiscardData(p);
   memset(&s, 0, sizeof(Fts5Structure));
   fts5DataWrite(p, FTS5_AVERAGES_ROWID, (const u8*)"", 0);
   fts5StructureWrite(p, &s);
@@ -214185,9 +219613,13 @@ static int sqlite3Fts5IndexCharlenToBytelen(
   for(i=0; i<nChar; i++){
     if( n>=nByte ) return 0;      /* Input contains fewer than nChar chars */
     if( (unsigned char)p[n++]>=0xc0 ){
+      if( n>=nByte ) return 0;
       while( (p[n] & 0xc0)==0x80 ){
         n++;
-        if( n>=nByte ) break;
+        if( n>=nByte ){
+          if( i+1==nChar ) break;
+          return 0;
+        }
       }
     }
   }
@@ -214323,7 +219755,7 @@ static int sqlite3Fts5IndexQuery(
     if( p->rc ){
       sqlite3Fts5IterClose((Fts5IndexIter*)pRet);
       pRet = 0;
-      fts5CloseReader(p);
+      sqlite3Fts5IndexCloseReader(p);
     }
 
     *ppIter = (Fts5IndexIter*)pRet;
@@ -214396,7 +219828,7 @@ static void sqlite3Fts5IterClose(Fts5IndexIter *pIndexIter){
     Fts5Iter *pIter = (Fts5Iter*)pIndexIter;
     Fts5Index *pIndex = pIter->pIndex;
     fts5MultiIterFree(pIter);
-    fts5CloseReader(pIndex);
+    sqlite3Fts5IndexCloseReader(pIndex);
   }
 }
 
@@ -214589,6 +220021,37 @@ static int fts5QueryCksum(
   return rc;
 }
 
+/*
+** Check if buffer z[], size n bytes, contains as series of valid utf-8
+** encoded codepoints. If so, return 0. Otherwise, if the buffer does not
+** contain valid utf-8, return non-zero.
+*/
+static int fts5TestUtf8(const char *z, int n){
+  int i = 0;
+  assert_nc( n>0 );
+  while( i<n ){
+    if( (z[i] & 0x80)==0x00 ){
+      i++;
+    }else
+    if( (z[i] & 0xE0)==0xC0 ){
+      if( i+1>=n || (z[i+1] & 0xC0)!=0x80 ) return 1;
+      i += 2;
+    }else
+    if( (z[i] & 0xF0)==0xE0 ){
+      if( i+2>=n || (z[i+1] & 0xC0)!=0x80 || (z[i+2] & 0xC0)!=0x80 ) return 1;
+      i += 3;
+    }else
+    if( (z[i] & 0xF8)==0xF0 ){
+      if( i+3>=n || (z[i+1] & 0xC0)!=0x80 || (z[i+2] & 0xC0)!=0x80 ) return 1;
+      if( (z[i+2] & 0xC0)!=0x80 ) return 1;
+      i += 3;
+    }else{
+      return 1;
+    }
+  }
+
+  return 0;
+}
 
 /*
 ** This function is also purely an internal test. It does not contribute to 
@@ -214629,8 +220092,14 @@ static void fts5TestTerm(
     ** This check may only be performed if the hash table is empty. This
     ** is because the hash table only supports a single scan query at
     ** a time, and the multi-iter loop from which this function is called
-    ** is already performing such a scan. */
-    if( p->nPendingData==0 ){
+    ** is already performing such a scan. 
+    **
+    ** Also only do this if buffer zTerm contains nTerm bytes of valid
+    ** utf-8. Otherwise, the last part of the buffer contents might contain
+    ** a non-utf-8 sequence that happens to be a prefix of a valid utf-8
+    ** character stored in the main fts index, which will cause the
+    ** test to fail.  */
+    if( p->nPendingData==0 && 0==fts5TestUtf8(zTerm, nTerm) ){
       if( iIdx>0 && rc==SQLITE_OK ){
         int f = flags|FTS5INDEX_QUERY_TEST_NOIDX;
         ck2 = 0;
@@ -214753,7 +220222,8 @@ static void fts5IndexIntegrityCheckSegment(
   if( pSeg->pgnoFirst==0 ) return;
 
   fts5IndexPrepareStmt(p, &pStmt, sqlite3_mprintf(
-      "SELECT segid, term, (pgno>>1), (pgno&1) FROM %Q.'%q_idx' WHERE segid=%d",
+      "SELECT segid, term, (pgno>>1), (pgno&1) FROM %Q.'%q_idx' WHERE segid=%d "
+      "ORDER BY 1, 2",
       pConfig->zDb, pConfig->zName, pSeg->iSegid
   ));
 
@@ -214762,8 +220232,8 @@ static void fts5IndexIntegrityCheckSegment(
     i64 iRow;                     /* Rowid for this leaf */
     Fts5Data *pLeaf;              /* Data for this leaf */
 
+    const char *zIdxTerm = (const char*)sqlite3_column_blob(pStmt, 1);
     int nIdxTerm = sqlite3_column_bytes(pStmt, 1);
-    const char *zIdxTerm = (const char*)sqlite3_column_text(pStmt, 1);
     int iIdxLeaf = sqlite3_column_int(pStmt, 2);
     int bIdxDlidx = sqlite3_column_int(pStmt, 3);
 
@@ -215744,7 +221214,10 @@ static void fts5CheckTransactionState(Fts5FullTable *p, int op, int iSavepoint){
     case FTS5_ROLLBACKTO:
       assert( p->ts.eState==1 );
       assert( iSavepoint>=-1 );
-      assert( iSavepoint<=p->ts.iSavepoint );
+      /* The following assert() can fail if another vtab strikes an error
+      ** within an xSavepoint() call then SQLite calls xRollbackTo() - without
+      ** having called xSavepoint() on this vtab.  */
+      /* assert( iSavepoint<=p->ts.iSavepoint ); */
       p->ts.iSavepoint = iSavepoint;
       break;
   }
@@ -215920,17 +221393,39 @@ static void fts5SetUniqueFlag(sqlite3_index_info *pIdxInfo){
 ** Implementation of the xBestIndex method for FTS5 tables. Within the 
 ** WHERE constraint, it searches for the following:
 **
-**   1. A MATCH constraint against the special column.
+**   1. A MATCH constraint against the table column.
 **   2. A MATCH constraint against the "rank" column.
-**   3. An == constraint against the rowid column.
-**   4. A < or <= constraint against the rowid column.
-**   5. A > or >= constraint against the rowid column.
+**   3. A MATCH constraint against some other column.
+**   4. An == constraint against the rowid column.
+**   5. A < or <= constraint against the rowid column.
+**   6. A > or >= constraint against the rowid column.
 **
-** Within the ORDER BY, either:
+** Within the ORDER BY, the following are supported:
 **
 **   5. ORDER BY rank [ASC|DESC]
 **   6. ORDER BY rowid [ASC|DESC]
 **
+** Information for the xFilter call is passed via both the idxNum and 
+** idxStr variables. Specifically, idxNum is a bitmask of the following
+** flags used to encode the ORDER BY clause:
+**
+**     FTS5_BI_ORDER_RANK
+**     FTS5_BI_ORDER_ROWID
+**     FTS5_BI_ORDER_DESC
+**
+** idxStr is used to encode data from the WHERE clause. For each argument
+** passed to the xFilter method, the following is appended to idxStr:
+**
+**   Match against table column:            "m"
+**   Match against rank column:             "r"
+**   Match against other column:            "<column-number>"
+**   Equality constraint against the rowid: "="
+**   A < or <= against the rowid:           "<"
+**   A > or >= against the rowid:           ">"
+**
+** This function ensures that there is at most one "r" or "=". And that if
+** there exists an "=" then there is no "<" or ">".
+**
 ** Costs are assigned as follows:
 **
 **  a) If an unusable MATCH operator is present in the WHERE clause, the
@@ -215958,32 +221453,18 @@ static int fts5BestIndexMethod(sqlite3_vtab *pVTab, sqlite3_index_info *pInfo){
   Fts5Config *pConfig = pTab->pConfig;
   const int nCol = pConfig->nCol;
   int idxFlags = 0;               /* Parameter passed through to xFilter() */
-  int bHasMatch;
-  int iNext;
   int i;
 
-  struct Constraint {
-    int op;                       /* Mask against sqlite3_index_constraint.op */
-    int fts5op;                   /* FTS5 mask for idxFlags */
-    int iCol;                     /* 0==rowid, 1==tbl, 2==rank */
-    int omit;                     /* True to omit this if found */
-    int iConsIndex;               /* Index in pInfo->aConstraint[] */
-  } aConstraint[] = {
-    {SQLITE_INDEX_CONSTRAINT_MATCH|SQLITE_INDEX_CONSTRAINT_EQ, 
-                                    FTS5_BI_MATCH,    1, 1, -1},
-    {SQLITE_INDEX_CONSTRAINT_MATCH|SQLITE_INDEX_CONSTRAINT_EQ, 
-                                    FTS5_BI_RANK,     2, 1, -1},
-    {SQLITE_INDEX_CONSTRAINT_EQ,    FTS5_BI_ROWID_EQ, 0, 0, -1},
-    {SQLITE_INDEX_CONSTRAINT_LT|SQLITE_INDEX_CONSTRAINT_LE, 
-                                    FTS5_BI_ROWID_LE, 0, 0, -1},
-    {SQLITE_INDEX_CONSTRAINT_GT|SQLITE_INDEX_CONSTRAINT_GE, 
-                                    FTS5_BI_ROWID_GE, 0, 0, -1},
-  };
+  char *idxStr;
+  int iIdxStr = 0;
+  int iCons = 0;
+
+  int bSeenEq = 0;
+  int bSeenGt = 0;
+  int bSeenLt = 0;
+  int bSeenMatch = 0;
+  int bSeenRank = 0;
 
-  int aColMap[3];
-  aColMap[0] = -1;
-  aColMap[1] = nCol;
-  aColMap[2] = nCol+1;
 
   assert( SQLITE_INDEX_CONSTRAINT_EQ<SQLITE_INDEX_CONSTRAINT_MATCH );
   assert( SQLITE_INDEX_CONSTRAINT_GT<SQLITE_INDEX_CONSTRAINT_MATCH );
@@ -215991,40 +221472,85 @@ static int fts5BestIndexMethod(sqlite3_vtab *pVTab, sqlite3_index_info *pInfo){
   assert( SQLITE_INDEX_CONSTRAINT_GE<SQLITE_INDEX_CONSTRAINT_MATCH );
   assert( SQLITE_INDEX_CONSTRAINT_LE<SQLITE_INDEX_CONSTRAINT_MATCH );
 
-  /* Set idxFlags flags for all WHERE clause terms that will be used. */
+  if( pConfig->bLock ){
+    pTab->base.zErrMsg = sqlite3_mprintf(
+        "recursively defined fts5 content table"
+    );
+    return SQLITE_ERROR;
+  }
+
+  idxStr = (char*)sqlite3_malloc(pInfo->nConstraint * 6 + 1);
+  if( idxStr==0 ) return SQLITE_NOMEM;
+  pInfo->idxStr = idxStr;
+  pInfo->needToFreeIdxStr = 1;
+
   for(i=0; i<pInfo->nConstraint; i++){
     struct sqlite3_index_constraint *p = &pInfo->aConstraint[i];
     int iCol = p->iColumn;
-
-    if( (p->op==SQLITE_INDEX_CONSTRAINT_MATCH && iCol>=0 && iCol<=nCol)
-     || (p->op==SQLITE_INDEX_CONSTRAINT_EQ && iCol==nCol)
+    if( p->op==SQLITE_INDEX_CONSTRAINT_MATCH
+     || (p->op==SQLITE_INDEX_CONSTRAINT_EQ && iCol>=nCol)
     ){
       /* A MATCH operator or equivalent */
-      if( p->usable ){
-        idxFlags = (idxFlags & 0xFFFF) | FTS5_BI_MATCH | (iCol << 16);
-        aConstraint[0].iConsIndex = i;
-      }else{
+      if( p->usable==0 || iCol<0 ){
         /* As there exists an unusable MATCH constraint this is an 
         ** unusable plan. Set a prohibitively high cost. */
         pInfo->estimatedCost = 1e50;
+        assert( iIdxStr < pInfo->nConstraint*6 + 1 );
+        idxStr[iIdxStr] = 0;
         return SQLITE_OK;
+      }else{
+        if( iCol==nCol+1 ){
+          if( bSeenRank ) continue;
+          idxStr[iIdxStr++] = 'r';
+          bSeenRank = 1;
+        }else{
+          bSeenMatch = 1;
+          idxStr[iIdxStr++] = 'm';
+          if( iCol<nCol ){
+            sqlite3_snprintf(6, &idxStr[iIdxStr], "%d", iCol);
+            idxStr += strlen(&idxStr[iIdxStr]);
+            assert( idxStr[iIdxStr]=='\0' );
+          }
+        }
+        pInfo->aConstraintUsage[i].argvIndex = ++iCons;
+        pInfo->aConstraintUsage[i].omit = 1;
       }
-    }else if( p->op<=SQLITE_INDEX_CONSTRAINT_MATCH ){
-      int j;
-      for(j=1; j<ArraySize(aConstraint); j++){
-        struct Constraint *pC = &aConstraint[j];
-        if( iCol==aColMap[pC->iCol] && (p->op & pC->op) && p->usable ){
-          pC->iConsIndex = i;
-          idxFlags |= pC->fts5op;
+    }
+    else if( p->usable && bSeenEq==0 
+      && p->op==SQLITE_INDEX_CONSTRAINT_EQ && iCol<0 
+    ){
+      idxStr[iIdxStr++] = '=';
+      bSeenEq = 1;
+      pInfo->aConstraintUsage[i].argvIndex = ++iCons;
+    }
+  }
+
+  if( bSeenEq==0 ){
+    for(i=0; i<pInfo->nConstraint; i++){
+      struct sqlite3_index_constraint *p = &pInfo->aConstraint[i];
+      if( p->iColumn<0 && p->usable ){
+        int op = p->op;
+        if( op==SQLITE_INDEX_CONSTRAINT_LT || op==SQLITE_INDEX_CONSTRAINT_LE ){
+          if( bSeenLt ) continue;
+          idxStr[iIdxStr++] = '<';
+          pInfo->aConstraintUsage[i].argvIndex = ++iCons;
+          bSeenLt = 1;
+        }else
+        if( op==SQLITE_INDEX_CONSTRAINT_GT || op==SQLITE_INDEX_CONSTRAINT_GE ){
+          if( bSeenGt ) continue;
+          idxStr[iIdxStr++] = '>';
+          pInfo->aConstraintUsage[i].argvIndex = ++iCons;
+          bSeenGt = 1;
         }
       }
     }
   }
+  idxStr[iIdxStr] = '\0';
 
   /* Set idxFlags flags for the ORDER BY clause */
   if( pInfo->nOrderBy==1 ){
     int iSort = pInfo->aOrderBy[0].iColumn;
-    if( iSort==(pConfig->nCol+1) && BitFlagTest(idxFlags, FTS5_BI_MATCH) ){
+    if( iSort==(pConfig->nCol+1) && bSeenMatch ){
       idxFlags |= FTS5_BI_ORDER_RANK;
     }else if( iSort==-1 ){
       idxFlags |= FTS5_BI_ORDER_ROWID;
@@ -216038,26 +221564,15 @@ static int fts5BestIndexMethod(sqlite3_vtab *pVTab, sqlite3_index_info *pInfo){
   }
 
   /* Calculate the estimated cost based on the flags set in idxFlags. */
-  bHasMatch = BitFlagTest(idxFlags, FTS5_BI_MATCH);
-  if( BitFlagTest(idxFlags, FTS5_BI_ROWID_EQ) ){
-    pInfo->estimatedCost = bHasMatch ? 100.0 : 10.0;
-    if( bHasMatch==0 ) fts5SetUniqueFlag(pInfo);
-  }else if( BitFlagAllTest(idxFlags, FTS5_BI_ROWID_LE|FTS5_BI_ROWID_GE) ){
-    pInfo->estimatedCost = bHasMatch ? 500.0 : 250000.0;
-  }else if( BitFlagTest(idxFlags, FTS5_BI_ROWID_LE|FTS5_BI_ROWID_GE) ){
-    pInfo->estimatedCost = bHasMatch ? 750.0 : 750000.0;
+  if( bSeenEq ){
+    pInfo->estimatedCost = bSeenMatch ? 100.0 : 10.0;
+    if( bSeenMatch==0 ) fts5SetUniqueFlag(pInfo);
+  }else if( bSeenLt && bSeenGt ){
+    pInfo->estimatedCost = bSeenMatch ? 500.0 : 250000.0;
+  }else if( bSeenLt || bSeenGt ){
+    pInfo->estimatedCost = bSeenMatch ? 750.0 : 750000.0;
   }else{
-    pInfo->estimatedCost = bHasMatch ? 1000.0 : 1000000.0;
-  }
-
-  /* Assign argvIndex values to each constraint in use. */
-  iNext = 1;
-  for(i=0; i<ArraySize(aConstraint); i++){
-    struct Constraint *pC = &aConstraint[i];
-    if( pC->iConsIndex>=0 ){
-      pInfo->aConstraintUsage[pC->iConsIndex].argvIndex = iNext++;
-      pInfo->aConstraintUsage[pC->iConsIndex].omit = (unsigned char)pC->omit;
-    }
+    pInfo->estimatedCost = bSeenMatch ? 1000.0 : 1000000.0;
   }
 
   pInfo->idxNum = idxFlags;
@@ -216157,6 +221672,7 @@ static void fts5FreeCursorComponents(Fts5Cursor *pCsr){
     sqlite3_free(pCsr->zRankArgs);
   }
 
+  sqlite3Fts5IndexCloseReader(pTab->p.pIndex);
   memset(&pCsr->ePlan, 0, sizeof(Fts5Cursor) - ((u8*)&pCsr->ePlan - (u8*)pCsr));
 }
 
@@ -216307,15 +221823,24 @@ static int fts5NextMethod(sqlite3_vtab_cursor *pCursor){
         break;
       }
   
-      default:
+      default: {
+        Fts5Config *pConfig = ((Fts5Table*)pCursor->pVtab)->pConfig;
+        pConfig->bLock++;
         rc = sqlite3_step(pCsr->pStmt);
+        pConfig->bLock--;
         if( rc!=SQLITE_ROW ){
           CsrFlagSet(pCsr, FTS5CSR_EOF);
           rc = sqlite3_reset(pCsr->pStmt);
+          if( rc!=SQLITE_OK ){
+            pCursor->pVtab->zErrMsg = sqlite3_mprintf(
+                "%s", sqlite3_errmsg(pConfig->db)
+            );
+          }
         }else{
           rc = SQLITE_OK;
         }
         break;
+      }
     }
   }
   
@@ -216380,7 +221905,7 @@ static int fts5CursorFirstSorted(
   **
   ** If SQLite a built-in statement cache, this wouldn't be a problem. */
   rc = fts5PrepareStatement(&pSorter->pStmt, pConfig,
-      "SELECT rowid, rank FROM %Q.%Q ORDER BY %s(%s%s%s) %s",
+      "SELECT rowid, rank FROM %Q.%Q ORDER BY %s(\"%w\"%s%s) %s",
       pConfig->zDb, pConfig->zName, zRank, pConfig->zName,
       (zRankArgs ? ", " : ""),
       (zRankArgs ? zRankArgs : ""),
@@ -216436,10 +221961,10 @@ static int fts5SpecialMatch(
   assert( pTab->p.base.zErrMsg==0 );
   pCsr->ePlan = FTS5_PLAN_SPECIAL;
 
-  if( 0==sqlite3_strnicmp("reads", z, n) ){
+  if( n==5 && 0==sqlite3_strnicmp("reads", z, n) ){
     pCsr->iSpecial = sqlite3Fts5IndexReads(pTab->p.pIndex);
   }
-  else if( 0==sqlite3_strnicmp("id", z, n) ){
+  else if( n==2 && 0==sqlite3_strnicmp("id", z, n) ){
     pCsr->iSpecial = pCsr->iCsrId;
   }
   else{
@@ -216580,7 +222105,7 @@ static i64 fts5GetRowidLimit(sqlite3_value *pVal, i64 iDefault){
 static int fts5FilterMethod(
   sqlite3_vtab_cursor *pCursor,   /* The cursor used for this query */
   int idxNum,                     /* Strategy index */
-  const char *zUnused,            /* Unused */
+  const char *idxStr,             /* Unused */
   int nVal,                       /* Number of elements in apVal */
   sqlite3_value **apVal           /* Arguments for the indexing scheme */
 ){
@@ -216588,19 +222113,24 @@ static int fts5FilterMethod(
   Fts5Config *pConfig = pTab->p.pConfig;
   Fts5Cursor *pCsr = (Fts5Cursor*)pCursor;
   int rc = SQLITE_OK;             /* Error code */
-  int iVal = 0;                   /* Counter for apVal[] */
   int bDesc;                      /* True if ORDER BY [rank|rowid] DESC */
   int bOrderByRank;               /* True if ORDER BY rank */
-  sqlite3_value *pMatch = 0;      /* <tbl> MATCH ? expression (or NULL) */
   sqlite3_value *pRank = 0;       /* rank MATCH ? expression (or NULL) */
   sqlite3_value *pRowidEq = 0;    /* rowid = ? expression (or NULL) */
   sqlite3_value *pRowidLe = 0;    /* rowid <= ? expression (or NULL) */
   sqlite3_value *pRowidGe = 0;    /* rowid >= ? expression (or NULL) */
   int iCol;                       /* Column on LHS of MATCH operator */
   char **pzErrmsg = pConfig->pzErrmsg;
+  int i;
+  int iIdxStr = 0;
+  Fts5Expr *pExpr = 0;
 
-  UNUSED_PARAM(zUnused);
-  UNUSED_PARAM(nVal);
+  if( pConfig->bLock ){
+    pTab->p.base.zErrMsg = sqlite3_mprintf(
+        "recursively defined fts5 content table"
+    );
+    return SQLITE_ERROR;
+  }
 
   if( pCsr->ePlan ){
     fts5FreeCursorComponents(pCsr);
@@ -216613,23 +222143,60 @@ static int fts5FilterMethod(
   assert( pCsr->pRank==0 );
   assert( pCsr->zRank==0 );
   assert( pCsr->zRankArgs==0 );
+  assert( pTab->pSortCsr==0 || nVal==0 );
 
   assert( pzErrmsg==0 || pzErrmsg==&pTab->p.base.zErrMsg );
   pConfig->pzErrmsg = &pTab->p.base.zErrMsg;
 
-  /* Decode the arguments passed through to this function.
-  **
-  ** Note: The following set of if(...) statements must be in the same
-  ** order as the corresponding entries in the struct at the top of
-  ** fts5BestIndexMethod().  */
-  if( BitFlagTest(idxNum, FTS5_BI_MATCH) ) pMatch = apVal[iVal++];
-  if( BitFlagTest(idxNum, FTS5_BI_RANK) ) pRank = apVal[iVal++];
-  if( BitFlagTest(idxNum, FTS5_BI_ROWID_EQ) ) pRowidEq = apVal[iVal++];
-  if( BitFlagTest(idxNum, FTS5_BI_ROWID_LE) ) pRowidLe = apVal[iVal++];
-  if( BitFlagTest(idxNum, FTS5_BI_ROWID_GE) ) pRowidGe = apVal[iVal++];
-  iCol = (idxNum>>16);
-  assert( iCol>=0 && iCol<=pConfig->nCol );
-  assert( iVal==nVal );
+  /* Decode the arguments passed through to this function. */
+  for(i=0; i<nVal; i++){
+    switch( idxStr[iIdxStr++] ){
+      case 'r':
+        pRank = apVal[i];
+        break;
+      case 'm': {
+        const char *zText = (const char*)sqlite3_value_text(apVal[i]);
+        if( zText==0 ) zText = "";
+
+        if( idxStr[iIdxStr]>='0' && idxStr[iIdxStr]<='9' ){
+          iCol = 0;
+          do{
+            iCol = iCol*10 + (idxStr[iIdxStr]-'0');
+            iIdxStr++;
+          }while( idxStr[iIdxStr]>='0' && idxStr[iIdxStr]<='9' );
+        }else{
+          iCol = pConfig->nCol;
+        }
+
+        if( zText[0]=='*' ){
+          /* The user has issued a query of the form "MATCH '*...'". This
+          ** indicates that the MATCH expression is not a full text query,
+          ** but a request for an internal parameter.  */
+          rc = fts5SpecialMatch(pTab, pCsr, &zText[1]);
+          goto filter_out;
+        }else{
+          char **pzErr = &pTab->p.base.zErrMsg;
+          rc = sqlite3Fts5ExprNew(pConfig, iCol, zText, &pExpr, pzErr);
+          if( rc==SQLITE_OK ){
+            rc = sqlite3Fts5ExprAnd(&pCsr->pExpr, pExpr);
+            pExpr = 0;
+          }
+          if( rc!=SQLITE_OK ) goto filter_out;
+        }
+
+        break;
+      }
+      case '=':
+        pRowidEq = apVal[i];
+        break;
+      case '<':
+        pRowidLe = apVal[i];
+        break;
+      default: assert( idxStr[iIdxStr-1]=='>' );
+        pRowidGe = apVal[i];
+        break;
+    }
+  }
   bOrderByRank = ((idxNum & FTS5_BI_ORDER_RANK) ? 1 : 0);
   pCsr->bDesc = bDesc = ((idxNum & FTS5_BI_ORDER_DESC) ? 1 : 0);
 
@@ -216656,7 +222223,7 @@ static int fts5FilterMethod(
     ** (pCursor) is used to execute the query issued by function 
     ** fts5CursorFirstSorted() above.  */
     assert( pRowidEq==0 && pRowidLe==0 && pRowidGe==0 && pRank==0 );
-    assert( nVal==0 && pMatch==0 && bOrderByRank==0 && bDesc==0 );
+    assert( nVal==0 && bOrderByRank==0 && bDesc==0 );
     assert( pCsr->iLastRowid==LARGEST_INT64 );
     assert( pCsr->iFirstRowid==SMALLEST_INT64 );
     if( pTab->pSortCsr->bDesc ){
@@ -216669,29 +222236,15 @@ static int fts5FilterMethod(
     pCsr->ePlan = FTS5_PLAN_SOURCE;
     pCsr->pExpr = pTab->pSortCsr->pExpr;
     rc = fts5CursorFirst(pTab, pCsr, bDesc);
-  }else if( pMatch ){
-    const char *zExpr = (const char*)sqlite3_value_text(apVal[0]);
-    if( zExpr==0 ) zExpr = "";
-
+  }else if( pCsr->pExpr ){
     rc = fts5CursorParseRank(pConfig, pCsr, pRank);
     if( rc==SQLITE_OK ){
-      if( zExpr[0]=='*' ){
-        /* The user has issued a query of the form "MATCH '*...'". This
-        ** indicates that the MATCH expression is not a full text query,
-        ** but a request for an internal parameter.  */
-        rc = fts5SpecialMatch(pTab, pCsr, &zExpr[1]);
+      if( bOrderByRank ){
+        pCsr->ePlan = FTS5_PLAN_SORTED_MATCH;
+        rc = fts5CursorFirstSorted(pTab, pCsr, bDesc);
       }else{
-        char **pzErr = &pTab->p.base.zErrMsg;
-        rc = sqlite3Fts5ExprNew(pConfig, iCol, zExpr, &pCsr->pExpr, pzErr);
-        if( rc==SQLITE_OK ){
-          if( bOrderByRank ){
-            pCsr->ePlan = FTS5_PLAN_SORTED_MATCH;
-            rc = fts5CursorFirstSorted(pTab, pCsr, bDesc);
-          }else{
-            pCsr->ePlan = FTS5_PLAN_MATCH;
-            rc = fts5CursorFirst(pTab, pCsr, bDesc);
-          }
-        }
+        pCsr->ePlan = FTS5_PLAN_MATCH;
+        rc = fts5CursorFirst(pTab, pCsr, bDesc);
       }
     }
   }else if( pConfig->zContent==0 ){
@@ -216708,7 +222261,7 @@ static int fts5FilterMethod(
     );
     if( rc==SQLITE_OK ){
       if( pCsr->ePlan==FTS5_PLAN_ROWID ){
-        sqlite3_bind_value(pCsr->pStmt, 1, apVal[0]);
+        sqlite3_bind_value(pCsr->pStmt, 1, pRowidEq);
       }else{
         sqlite3_bind_int64(pCsr->pStmt, 1, pCsr->iFirstRowid);
         sqlite3_bind_int64(pCsr->pStmt, 2, pCsr->iLastRowid);
@@ -216717,6 +222270,8 @@ static int fts5FilterMethod(
     }
   }
 
+ filter_out:
+  sqlite3Fts5ExprFree(pExpr);
   pConfig->pzErrmsg = pzErrmsg;
   return rc;
 }
@@ -216797,10 +222352,13 @@ static int fts5SeekCursor(Fts5Cursor *pCsr, int bErrormsg){
   }
 
   if( rc==SQLITE_OK && CsrFlagTest(pCsr, FTS5CSR_REQUIRE_CONTENT) ){
+    Fts5Table *pTab = (Fts5Table*)(pCsr->base.pVtab);
     assert( pCsr->pExpr );
     sqlite3_reset(pCsr->pStmt);
     sqlite3_bind_int64(pCsr->pStmt, 1, fts5CursorRowid(pCsr));
+    pTab->pConfig->bLock++;
     rc = sqlite3_step(pCsr->pStmt);
+    pTab->pConfig->bLock--;
     if( rc==SQLITE_ROW ){
       rc = SQLITE_OK;
       CsrFlagClear(pCsr, FTS5CSR_REQUIRE_CONTENT);
@@ -216808,6 +222366,10 @@ static int fts5SeekCursor(Fts5Cursor *pCsr, int bErrormsg){
       rc = sqlite3_reset(pCsr->pStmt);
       if( rc==SQLITE_OK ){
         rc = FTS5_CORRUPT;
+      }else if( pTab->pConfig->pzErrmsg ){
+        *pTab->pConfig->pzErrmsg = sqlite3_mprintf(
+            "%s", sqlite3_errmsg(pTab->pConfig->db)
+        );
       }
     }
   }
@@ -217687,7 +223249,7 @@ static void fts5ApiCallback(
   iCsrId = sqlite3_value_int64(argv[0]);
 
   pCsr = fts5CursorFromCsrid(pAux->pGlobal, iCsrId);
-  if( pCsr==0 ){
+  if( pCsr==0 || pCsr->ePlan==0 ){
     char *zErr = sqlite3_mprintf("no such cursor: %lld", iCsrId);
     sqlite3_result_error(context, zErr, -1);
     sqlite3_free(zErr);
@@ -217823,10 +223385,12 @@ static int fts5ColumnMethod(
       }
     }
   }else if( !fts5IsContentless(pTab) ){
+    pConfig->pzErrmsg = &pTab->p.base.zErrMsg;
     rc = fts5SeekCursor(pCsr, 1);
     if( rc==SQLITE_OK ){
       sqlite3_result_value(pCtx, sqlite3_column_value(pCsr->pStmt, iCol+1));
     }
+    pConfig->pzErrmsg = 0;
   }
   return rc;
 }
@@ -218103,7 +223667,7 @@ static void fts5SourceIdFunc(
 ){
   assert( nArg==0 );
   UNUSED_PARAM2(nArg, apUnused);
-  sqlite3_result_text(pCtx, "fts5: 2019-04-16 19:49:53 884b4b7e502b4e991677b53971277adfaf0a04a284f8e483e2553d0f83156b50", -1, SQLITE_TRANSIENT);
+  sqlite3_result_text(pCtx, "fts5: 2020-01-27 19:55:54 3bfa9cc97da10598521b342961df8f5f68c7388fa117345eeb516eaa837bb4d6", -1, SQLITE_TRANSIENT);
 }
 
 /*
@@ -218375,7 +223939,9 @@ static int fts5StorageGetStmt(
     }else{
       int f = SQLITE_PREPARE_PERSISTENT;
       if( eStmt>FTS5_STMT_LOOKUP ) f |= SQLITE_PREPARE_NO_VTAB;
+      p->pConfig->bLock++;
       rc = sqlite3_prepare_v3(pC->db, zSql, -1, f, &p->aStmt[eStmt], 0);
+      p->pConfig->bLock--;
       sqlite3_free(zSql);
       if( rc!=SQLITE_OK && pzErrMsg ){
         *pzErrMsg = sqlite3_mprintf("%s", sqlite3_errmsg(pC->db));
@@ -218795,6 +224361,8 @@ static int sqlite3Fts5StorageDeleteAll(Fts5Storage *p){
   Fts5Config *pConfig = p->pConfig;
   int rc;
 
+  p->bTotalsValid = 0;
+
   /* Delete the contents of the %_data and %_docsize tables. */
   rc = fts5ExecPrintf(pConfig->db, 0,
       "DELETE FROM %Q.'%q_data';" 
@@ -218846,10 +224414,11 @@ static int sqlite3Fts5StorageRebuild(Fts5Storage *p){
     for(ctx.iCol=0; rc==SQLITE_OK && ctx.iCol<pConfig->nCol; ctx.iCol++){
       ctx.szCol = 0;
       if( pConfig->abUnindexed[ctx.iCol]==0 ){
+        const char *zText = (const char*)sqlite3_column_text(pScan, ctx.iCol+1);
+        int nText = sqlite3_column_bytes(pScan, ctx.iCol+1);
         rc = sqlite3Fts5Tokenize(pConfig, 
             FTS5_TOKENIZE_DOCUMENT,
-            (const char*)sqlite3_column_text(pScan, ctx.iCol+1),
-            sqlite3_column_bytes(pScan, ctx.iCol+1),
+            zText, nText,
             (void*)&ctx,
             fts5StorageInsertCallback
         );
@@ -218971,10 +224540,11 @@ static int sqlite3Fts5StorageIndexInsert(
   for(ctx.iCol=0; rc==SQLITE_OK && ctx.iCol<pConfig->nCol; ctx.iCol++){
     ctx.szCol = 0;
     if( pConfig->abUnindexed[ctx.iCol]==0 ){
+      const char *zText = (const char*)sqlite3_value_text(apVal[ctx.iCol+2]);
+      int nText = sqlite3_value_bytes(apVal[ctx.iCol+2]);
       rc = sqlite3Fts5Tokenize(pConfig, 
           FTS5_TOKENIZE_DOCUMENT,
-          (const char*)sqlite3_value_text(apVal[ctx.iCol+2]),
-          sqlite3_value_bytes(apVal[ctx.iCol+2]),
+          zText, nText,
           (void*)&ctx,
           fts5StorageInsertCallback
       );
@@ -219143,10 +224713,11 @@ static int sqlite3Fts5StorageIntegrity(Fts5Storage *p){
           rc = sqlite3Fts5TermsetNew(&ctx.pTermset);
         }
         if( rc==SQLITE_OK ){
+          const char *zText = (const char*)sqlite3_column_text(pScan, i+1);
+          int nText = sqlite3_column_bytes(pScan, i+1);
           rc = sqlite3Fts5Tokenize(pConfig, 
               FTS5_TOKENIZE_DOCUMENT,
-              (const char*)sqlite3_column_text(pScan, i+1),
-              sqlite3_column_bytes(pScan, i+1),
+              zText, nText,
               (void*)&ctx,
               fts5StorageIntegrityCallback
           );
@@ -222362,8 +227933,10 @@ static int fts5VocabNextMethod(sqlite3_vtab_cursor *pCursor){
   }
 
   if( rc==SQLITE_OK && pCsr->bEof==0 && pTab->eType==FTS5_VOCAB_COL ){
-    while( pCsr->aDoc[pCsr->iCol]==0 ) pCsr->iCol++;
-    assert( pCsr->iCol<pCsr->pFts5->pConfig->nCol );
+    for(/* noop */; pCsr->iCol<nCol && pCsr->aDoc[pCsr->iCol]==0; pCsr->iCol++);
+    if( pCsr->iCol==nCol ){
+      rc = FTS5_CORRUPT;
+    }
   }
   return rc;
 }
@@ -222867,9 +228440,9 @@ SQLITE_API int sqlite3_stmt_init(
 #endif /* !defined(SQLITE_CORE) || defined(SQLITE_ENABLE_STMTVTAB) */
 
 /************** End of stmt.c ************************************************/
-#if __LINE__!=222870
+#if __LINE__!=228443
 #undef SQLITE_SOURCE_ID
-#define SQLITE_SOURCE_ID      "2019-04-16 19:49:53 884b4b7e502b4e991677b53971277adfaf0a04a284f8e483e2553d0f8315alt2"
+#define SQLITE_SOURCE_ID      "2020-01-27 19:55:54 3bfa9cc97da10598521b342961df8f5f68c7388fa117345eeb516eaa837balt2"
 #endif
 /* Return the source-id for this library */
 SQLITE_API const char *sqlite3_sourceid(void){ return SQLITE_SOURCE_ID; }
