Updating SQLite pull to newest release.  Enabling additional #defines to allow for JSON support.

diff --git a/Engine/Plugins/Runtime/Database/SQLiteCore/Source/SQLiteCore/SQLiteCore.Build.cs b/Engine/Plugins/Runtime/Database/SQLiteCore/Source/SQLiteCore/SQLiteCore.Build.cs
--- a/Engine/Plugins/Runtime/Database/SQLiteCore/Source/SQLiteCore/SQLiteCore.Build.cs
+++ b/Engine/Plugins/Runtime/Database/SQLiteCore/Source/SQLiteCore/SQLiteCore.Build.cs
@@ -26,6 +26,9 @@ namespace UnrealBuildTool.Rules
 			PrivateDefinitions.Add("SQLITE_ENABLE_FTS4");
 			PrivateDefinitions.Add("SQLITE_ENABLE_FTS5");
 
+			// Enable Json extension
+			PrivateDefinitions.Add("SQLITE_ENABLE_JSON1");
+			
 			// Use ICU with SQLite if it's available
 			if (Target.bCompileICU)
 			{
