Remove local SQLite from Mockup Mantle and fixed build errors using the internal version in SQLiteCore. There is a temporary fix  in there to get the linkage working that will need to be removed.

diff --git a/Engine/Plugins/Runtime/Database/SQLiteCore/Source/SQLiteCore/Private/SQLiteEmbedded.c b/Engine/Plugins/Runtime/Database/SQLiteCore/Source/SQLiteCore/Private/SQLiteEmbedded.c
--- a/Engine/Plugins/Runtime/Database/SQLiteCore/Source/SQLiteCore/Private/SQLiteEmbedded.c
+++ b/Engine/Plugins/Runtime/Database/SQLiteCore/Source/SQLiteCore/Private/SQLiteEmbedded.c
@@ -1,15 +1,27 @@
 // Copyright Epic Games, Inc. All Rights Reserved.
 
+
 /** Needed to define USING_CODE_ANALYSIS */
 #if PLATFORM_WINDOWS
 #include "Windows/WindowsPlatformCodeAnalysis.h"
+// This is a temporary fix to directly access SQLite API from outside this module.
+// DLLEXPORT and DLLImport are the ones to remove
+#define DLLEXPORT __declspec(dllexport)
+#define DLLIMPORT __declspec(dllimport)
 #elif defined(__clang__)
 #include "Clang/ClangPlatformCodeAnalysis.h"
+// This is a fix to directly access SQLite API from outside this module. 
+// DLLEXPORT and DLLImport are the ones to remove
+#define DLLEXPORT __attribute__((visibility("default")))
+#define DLLIMPORT __attribute__((visibility("default")))
 #endif
 
+ 
+
 /** Include SQLite, but not if we're building for analysis as the code emits warnings */
 #if !USING_CODE_ANALYSIS
 #include "HAL/PreprocessorHelpers.h"
 #include COMPILED_PLATFORM_HEADER(PlatformCompilerPreSetup.h)
 #include "sqlite/sqlite3.inl"
+
 #endif
diff --git a/Engine/Plugins/Runtime/Database/SQLiteCore/Source/SQLiteCore/SQLiteCore.Build.cs b/Engine/Plugins/Runtime/Database/SQLiteCore/Source/SQLiteCore/SQLiteCore.Build.cs
--- a/Engine/Plugins/Runtime/Database/SQLiteCore/Source/SQLiteCore/SQLiteCore.Build.cs
+++ b/Engine/Plugins/Runtime/Database/SQLiteCore/Source/SQLiteCore/SQLiteCore.Build.cs
@@ -26,6 +26,14 @@ namespace UnrealBuildTool.Rules
 			PrivateDefinitions.Add("SQLITE_ENABLE_FTS4");
 			PrivateDefinitions.Add("SQLITE_ENABLE_FTS5");
 
+			// Enable Deserialization and RTree Subsystem
+			PrivateDefinitions.Add("SQLITE_ENABLE_DESERIALIZE");
+			PrivateDefinitions.Add("SQLITE_ENABLE_RTREE");
+
+			// This is a temporary fix to directly access SQLite API from outside this module. 
+			// There are two API macros in SQLiteEmbedded.c that also need to be removed			
+		    PublicDefinitions.Add("SQLITE_API=SQLITECORE_API");
+
 			// Enable Json extension
 			PrivateDefinitions.Add("SQLITE_ENABLE_JSON1");
 			
