Wrapped SQLite in TPS macros to suppress warnings

#rnx

diff --git a/Engine/Plugins/Runtime/Database/SQLiteCore/Source/SQLiteCore/Private/SQLiteEmbedded.c b/Engine/Plugins/Runtime/Database/SQLiteCore/Source/SQLiteCore/Private/SQLiteEmbedded.c
--- a/Engine/Plugins/Runtime/Database/SQLiteCore/Source/SQLiteCore/Private/SQLiteEmbedded.c
+++ b/Engine/Plugins/Runtime/Database/SQLiteCore/Source/SQLiteCore/Private/SQLiteEmbedded.c
@@ -1,6 +1,5 @@
 // Copyright Epic Games, Inc. All Rights Reserved.
 
-
 /** Needed to define USING_CODE_ANALYSIS */
 #if PLATFORM_WINDOWS
 #include "Windows/WindowsPlatformCodeAnalysis.h"
@@ -16,12 +15,11 @@
 #define DLLIMPORT __attribute__((visibility("default")))
 #endif
 
- 
-
 /** Include SQLite, but not if we're building for analysis as the code emits warnings */
 #if !USING_CODE_ANALYSIS
 #include "HAL/PreprocessorHelpers.h"
 #include COMPILED_PLATFORM_HEADER(PlatformCompilerPreSetup.h)
+THIRD_PARTY_INCLUDES_START
 #include "sqlite/sqlite3.inl"
-
+THIRD_PARTY_INCLUDES_END
 #endif
