Allow direct access to the SQLite API via SQLiteCore

diff --git a/Engine/Plugins/Runtime/Database/SQLiteCore/Source/SQLiteCore/Private/SQLiteEmbedded.c b/Engine/Plugins/Runtime/Database/SQLiteCore/Source/SQLiteCore/Private/SQLiteEmbedded.c
--- a/Engine/Plugins/Runtime/Database/SQLiteCore/Source/SQLiteCore/Private/SQLiteEmbedded.c
+++ b/Engine/Plugins/Runtime/Database/SQLiteCore/Source/SQLiteCore/Private/SQLiteEmbedded.c
@@ -1,16 +1,12 @@
 // Copyright Epic Games, Inc. All Rights Reserved.
 
-/** Needed to define USING_CODE_ANALYSIS */
+/** Needed to define USING_CODE_ANALYSIS, and allow direct access to the SQLite API from outside this module */
 #if PLATFORM_WINDOWS
 #include "Windows/WindowsPlatformCodeAnalysis.h"
-// This is a temporary fix to directly access SQLite API from outside this module.
-// DLLEXPORT and DLLImport are the ones to remove
 #define DLLEXPORT __declspec(dllexport)
 #define DLLIMPORT __declspec(dllimport)
 #elif defined(__clang__)
 #include "Clang/ClangPlatformCodeAnalysis.h"
-// This is a fix to directly access SQLite API from outside this module. 
-// DLLEXPORT and DLLImport are the ones to remove
 #define DLLEXPORT __attribute__((visibility("default")))
 #define DLLIMPORT __attribute__((visibility("default")))
 #endif
diff --git a/Engine/Plugins/Runtime/Database/SQLiteCore/Source/SQLiteCore/Private/sqlite/README.txt b/Engine/Plugins/Runtime/Database/SQLiteCore/Source/SQLiteCore/Private/sqlite/README.txt
--- a/Engine/Plugins/Runtime/Database/SQLiteCore/Source/SQLiteCore/Private/sqlite/README.txt
+++ b/Engine/Plugins/Runtime/Database/SQLiteCore/Source/SQLiteCore/Private/sqlite/README.txt
@@ -2,5 +2,5 @@ This file contains the SQLite amalgamation that we embed into the SQLiteCore mod
 
 To change the version of SQLite that is embedded, you need to:
     1) Go to https://www.sqlite.org/download.html and download a new amalgamation.
-	2) Copy sqlite3.h and sqlite3.c into this folder.
+	2) Copy sqlite3.h into SQLiteCore\Public\sqlite and sqlite3.c into SQLiteCore\Private\sqlite.
 	3) Rename sqlite3.c to sqlite3.inl.
