Allow direct access to the SQLite API via SQLiteCore

diff --git a/Engine/Plugins/Runtime/Database/SQLiteCore/Source/SQLiteCore/SQLiteCore.Build.cs b/Engine/Plugins/Runtime/Database/SQLiteCore/Source/SQLiteCore/SQLiteCore.Build.cs
--- a/Engine/Plugins/Runtime/Database/SQLiteCore/Source/SQLiteCore/SQLiteCore.Build.cs
+++ b/Engine/Plugins/Runtime/Database/SQLiteCore/Source/SQLiteCore/SQLiteCore.Build.cs
@@ -14,7 +14,12 @@ namespace UnrealBuildTool.Rules
 				}
 			);
 
-			PrivateDefinitions.Add("SQLITE_OMIT_AUTOINIT"); // We call sqlite3_initialize ourselves during module init
+			// Allow direct access to the SQLite API from outside this module
+			// Note: There are additional macros defined inside SQLiteEmbedded.c to allow this API export to work
+			PublicDefinitions.Add("SQLITE_API=SQLITECORE_API");
+
+			// We call sqlite3_initialize ourselves during module init
+			PrivateDefinitions.Add("SQLITE_OMIT_AUTOINIT");
 
 			// Use the math.h version of isnan rather than the SQLite version to avoid a -ffast-math error
 			PrivateDefinitions.Add("SQLITE_HAVE_ISNAN=1");
@@ -30,10 +35,6 @@ namespace UnrealBuildTool.Rules
 			PrivateDefinitions.Add("SQLITE_ENABLE_DESERIALIZE");
 			PrivateDefinitions.Add("SQLITE_ENABLE_RTREE");
 
-			// This is a temporary fix to directly access SQLite API from outside this module. 
-			// There are two API macros in SQLiteEmbedded.c that also need to be removed			
-		    PublicDefinitions.Add("SQLITE_API=SQLITECORE_API");
-
 			// Enable Json extension
 			PrivateDefinitions.Add("SQLITE_ENABLE_JSON1");
 			
