PR #8680: [SQLiteCore] Fix compile error when SQLITE_DEBUG enabled (Contributed by [REDACTED_PERSON])

The mentioned error is a clang error. It does not repro in cl.exe. Reading over this explanation of temp copies and constructor choice: https://stackoverflow.com/questions/21708606/why-does-an-in-place-member-initialization-use-a-copy-constructor-in-c11, initializer list syntax seems like a resonable solution. Delegating a call to the TAtomic<> ctor in the struct ctor init list would just add more #ifdef and hurt reability. Accepting.

#rnx

diff --git a/Engine/Plugins/Runtime/Database/SQLiteCore/Source/SQLiteCore/Private/SQLiteEmbeddedPlatform.cpp b/Engine/Plugins/Runtime/Database/SQLiteCore/Source/SQLiteCore/Private/SQLiteEmbeddedPlatform.cpp
--- a/Engine/Plugins/Runtime/Database/SQLiteCore/Source/SQLiteCore/Private/SQLiteEmbeddedPlatform.cpp
+++ b/Engine/Plugins/Runtime/Database/SQLiteCore/Source/SQLiteCore/Private/SQLiteEmbeddedPlatform.cpp
@@ -90,7 +90,7 @@ struct FSQLiteMutex
 	FCriticalSection CriticalSection;
 	int SQLiteMutexId;
 #ifdef SQLITE_DEBUG
-	TAtomic<uint32> OwnerThreadId = (uint32)INDEX_NONE;
+	TAtomic<uint32> OwnerThreadId{(uint32)INDEX_NONE};
 #endif
 };
 
