Updated SQLiteCore to compile with Unity files and not have bEnableUndefinedIdentifierWarnings as false.

With the current SQLiteCore settings, a 216MB shared PCH gets created specifically for this module.

diff --git a/Engine/Plugins/Runtime/Database/SQLiteCore/Source/SQLiteCore/Private/SQLiteEmbedded.c b/Engine/Plugins/Runtime/Database/SQLiteCore/Source/SQLiteCore/Private/SQLiteEmbedded.c
--- a/Engine/Plugins/Runtime/Database/SQLiteCore/Source/SQLiteCore/Private/SQLiteEmbedded.c
+++ b/Engine/Plugins/Runtime/Database/SQLiteCore/Source/SQLiteCore/Private/SQLiteEmbedded.c
@@ -1,7 +1,7 @@
 // Copyright Epic Games, Inc. All Rights Reserved.
 
 /** Needed to define USING_CODE_ANALYSIS, and allow direct access to the SQLite API from outside this module */
-#if PLATFORM_WINDOWS
+#if defined(PLATFORM_WINDOWS) && PLATFORM_WINDOWS
 #include "Windows/WindowsPlatformCodeAnalysis.h"
 #define DLLEXPORT __declspec(dllexport)
 #define DLLIMPORT __declspec(dllimport)
@@ -12,7 +12,7 @@
 #endif
 
 /** Include SQLite, but not if we're building for analysis as the code emits warnings */
-#if !USING_CODE_ANALYSIS
+#if !defined(USING_CODE_ANALYSIS) || !USING_CODE_ANALYSIS
 #include "HAL/PreprocessorHelpers.h"
 #include COMPILED_PLATFORM_HEADER(PlatformCompilerPreSetup.h)
 THIRD_PARTY_INCLUDES_START
diff --git a/Engine/Plugins/Runtime/Database/SQLiteCore/Source/SQLiteCore/SQLiteCore.Build.cs b/Engine/Plugins/Runtime/Database/SQLiteCore/Source/SQLiteCore/SQLiteCore.Build.cs
--- a/Engine/Plugins/Runtime/Database/SQLiteCore/Source/SQLiteCore/SQLiteCore.Build.cs
+++ b/Engine/Plugins/Runtime/Database/SQLiteCore/Source/SQLiteCore/SQLiteCore.Build.cs
@@ -56,9 +56,6 @@ namespace UnrealBuildTool.Rules
 				PrivateDefinitions.Add("SQLITE_MUTEX_NOOP");			// We provide our own mutex implementation
 				PrivateDefinitions.Add("SQLITE_OMIT_LOAD_EXTENSION");	// We disable extension loading
 			}
-
-			bEnableUndefinedIdentifierWarnings = false; // The embedded SQLite implementation generates a lot of these warnings
-			bUseUnity = false; // Ensure the embedded SQLite implementation is always compiled in isolation
 		}
 	}
 }
