fix SQLiteCore modular build link errors on some Microsoft platforms.

#rnx

diff --git a/Engine/Plugins/Runtime/Database/SQLiteCore/Source/SQLiteCore/Private/SQLiteEmbedded.c b/Engine/Plugins/Runtime/Database/SQLiteCore/Source/SQLiteCore/Private/SQLiteEmbedded.c
--- a/Engine/Plugins/Runtime/Database/SQLiteCore/Source/SQLiteCore/Private/SQLiteEmbedded.c
+++ b/Engine/Plugins/Runtime/Database/SQLiteCore/Source/SQLiteCore/Private/SQLiteEmbedded.c
@@ -1,8 +1,10 @@
 // Copyright Epic Games, Inc. All Rights Reserved.
 
 /** Needed to define USING_CODE_ANALYSIS, and allow direct access to the SQLite API from outside this module */
-#if defined(PLATFORM_WINDOWS) && PLATFORM_WINDOWS
-#include "Windows/WindowsPlatformCodeAnalysis.h"
+#if defined(PLATFORM_MICROSOFT) && PLATFORM_MICROSOFT
+ #if defined(PLATFORM_WINDOWS) && PLATFORM_WINDOWS
+ #include "Windows/WindowsPlatformCodeAnalysis.h"
+ #endif
 #define DLLEXPORT __declspec(dllexport)
 #define DLLIMPORT __declspec(dllimport)
 #elif defined(__clang__)
