Insights markers

diff --git a/Engine/Plugins/Runtime/Database/SQLiteCore/Source/SQLiteCore/Private/SQLiteDatabase.cpp b/Engine/Plugins/Runtime/Database/SQLiteCore/Source/SQLiteCore/Private/SQLiteDatabase.cpp
--- a/Engine/Plugins/Runtime/Database/SQLiteCore/Source/SQLiteCore/Private/SQLiteDatabase.cpp
+++ b/Engine/Plugins/Runtime/Database/SQLiteCore/Source/SQLiteCore/Private/SQLiteDatabase.cpp
@@ -8,6 +8,8 @@
 #include "HAL/FileManager.h"
 #include "Misc/AssertionMacros.h"
 #include "Containers/StringConv.h"
+#include "Trace/Trace.h"
+#include "Trace/Trace.inl"
 
 DEFINE_LOG_CATEGORY_STATIC(LogSQLiteDatabase, Log, All);
 
@@ -52,6 +54,9 @@ bool FSQLiteDatabase::Open(const TCHAR* InFilename, const ESQLiteDatabaseOpenMod
 	{
 		return false;
 	}
+	
+	OriginalPath = InFilename;
+	ShortName = FPaths::GetBaseFilename(InFilename);
 
 	int32 OpenFlags = 0;
 	switch (InOpenMode)
@@ -86,6 +91,7 @@ bool FSQLiteDatabase::Open(const TCHAR* InFilename, const ESQLiteDatabaseOpenMod
 		}
 		return false;
 	}
+	UE_LOG(LogSQLiteDatabase, Log, TEXT("Opened database '%s'"), InFilename);
 
 	return true;
 }
@@ -205,8 +211,16 @@ int64 FSQLiteDatabase::GetLastInsertRowId() const
 		: 0;
 }
 
+UE_TRACE_EVENT_BEGIN(Cpu, Sqlite_PerformQuickIntegrityCheck, NoSync)
+	UE_TRACE_EVENT_FIELD(UE::Trace::WideString, Path)
+UE_TRACE_EVENT_END()
+
 bool FSQLiteDatabase::PerformQuickIntegrityCheck() const
 {
+#if CPUPROFILERTRACE_ENABLED
+	UE_TRACE_LOG_SCOPED_T(Cpu, Sqlite_PerformQuickIntegrityCheck, CpuChannel)
+		<< Sqlite_PerformQuickIntegrityCheck.Path(*OriginalPath);
+#endif // CPUPROFILERTRACE_ENABLED
 	bool OutIntegrityOk = true;
 	bool bSuccessful = const_cast<FSQLiteDatabase*>(this)->Execute(TEXT("pragma quick_check;"), [&OutIntegrityOk](const FSQLitePreparedStatement& InStatement)
 	{
diff --git a/Engine/Plugins/Runtime/Database/SQLiteCore/Source/SQLiteCore/Private/SQLitePreparedStatement.cpp b/Engine/Plugins/Runtime/Database/SQLiteCore/Source/SQLiteCore/Private/SQLitePreparedStatement.cpp
--- a/Engine/Plugins/Runtime/Database/SQLiteCore/Source/SQLiteCore/Private/SQLitePreparedStatement.cpp
+++ b/Engine/Plugins/Runtime/Database/SQLiteCore/Source/SQLiteCore/Private/SQLitePreparedStatement.cpp
@@ -392,6 +392,7 @@ bool FSQLitePreparedStatement::Execute()
 
 int64 FSQLitePreparedStatement::Execute(TFunctionRef<ESQLitePreparedStatementExecuteRowResult(const FSQLitePreparedStatement&)> InCallback)
 {
+	TRACE_CPUPROFILER_EVENT_SCOPE_STR("FSQLitePreparedStatement::Execute");
 	checkf(IsValid() && !IsActive(), TEXT("SQLite statement must be valid and not-active!"));
 
 	// Step it to completion (or error)
diff --git a/Engine/Plugins/Runtime/Database/SQLiteCore/Source/SQLiteCore/Public/SQLiteDatabase.h b/Engine/Plugins/Runtime/Database/SQLiteCore/Source/SQLiteCore/Public/SQLiteDatabase.h
--- a/Engine/Plugins/Runtime/Database/SQLiteCore/Source/SQLiteCore/Public/SQLiteDatabase.h
+++ b/Engine/Plugins/Runtime/Database/SQLiteCore/Source/SQLiteCore/Public/SQLiteDatabase.h
@@ -139,4 +139,10 @@ private:
 
 	/** Internal SQLite database handle */
 	struct sqlite3* Database;
+	
+	/** Full original path for logging/profiling */
+	FString OriginalPath;
+
+	/** Short name for logging/profiling */
+	FString ShortName;
 };
diff --git a/Engine/Plugins/Runtime/Database/SQLiteCore/Source/SQLiteCore/SQLiteCore.Build.cs b/Engine/Plugins/Runtime/Database/SQLiteCore/Source/SQLiteCore/SQLiteCore.Build.cs
--- a/Engine/Plugins/Runtime/Database/SQLiteCore/Source/SQLiteCore/SQLiteCore.Build.cs
+++ b/Engine/Plugins/Runtime/Database/SQLiteCore/Source/SQLiteCore/SQLiteCore.Build.cs
@@ -13,6 +13,7 @@ namespace UnrealBuildTool.Rules
 			PublicDependencyModuleNames.AddRange(
 				new string[] {
 					"Core",
+					"TraceLog"
 				}
 			);
 
