Fixed crash in SQLite when the current malloc doesn't support FMemory::GetAllocSize.

diff --git a/Engine/Plugins/Runtime/Database/SQLiteCore/Source/SQLiteCore/Private/SQLiteEmbeddedPlatform.cpp b/Engine/Plugins/Runtime/Database/SQLiteCore/Source/SQLiteCore/Private/SQLiteEmbeddedPlatform.cpp
--- a/Engine/Plugins/Runtime/Database/SQLiteCore/Source/SQLiteCore/Private/SQLiteEmbeddedPlatform.cpp
+++ b/Engine/Plugins/Runtime/Database/SQLiteCore/Source/SQLiteCore/Private/SQLiteEmbeddedPlatform.cpp
@@ -32,8 +32,24 @@ public:
 			&Shutdown,
 			nullptr,
 		};
+		static const sqlite3_mem_methods WorkaroundMallocFuncs = {
+			&WorkaroundMalloc,
+			&WorkaroundFree,
+			&WorkaroundRealloc,
+			&WorkaroundSize,
+			&WorkaroundRoundup,
+			&Init,
+			&Shutdown,
+			nullptr,
+		};
+
+		// SQLite needs a working FMemory::GetAllocSize, so check if the current allocator supports FMemory::GetAllocSize()
+		void* TempMem = FMemory::Malloc(16);
+		bool bSupportsGetAllocSize = FMemory::GetAllocSize(TempMem) >= 16;
+		FMemory::Free(TempMem);
 
-		sqlite3_config(SQLITE_CONFIG_MALLOC, &MallocFuncs);
+		// Use the workaround functions if GetAllocSize is not supported.
+		sqlite3_config(SQLITE_CONFIG_MALLOC, bSupportsGetAllocSize ? &MallocFuncs : &WorkaroundMallocFuncs);
 	}
 
 private:
@@ -77,6 +93,61 @@ private:
 	{
 		return (int)FMemory::QuantizeSize(InSizeBytes, DEFAULT_ALIGNMENT);
 	}
+
+	/**
+	 * Workaround memory allocation functions which store the size in the 16 bytes preceding the allocation
+	 * This assumes that the default alignment is no more than 16 bytes, so we can just subtract 16 from any non-null pointer to get to the size
+	 */
+	static constexpr int WorkaroundHeaderSize = 16;
+
+	/** Allocate memory */
+	static void* WorkaroundMalloc(int InSizeBytes)
+	{
+		void* Result = FMemory::Malloc(InSizeBytes + WorkaroundHeaderSize, DEFAULT_ALIGNMENT);
+		if (Result)
+		{
+			*(int*)Result = InSizeBytes;
+			Result = (char*)Result + WorkaroundHeaderSize;
+		}
+		return Result;
+	}
+
+	/** Free memory returned by Alloc or Realloc */
+	static void WorkaroundFree(void* InPtr)
+	{
+		if (InPtr)
+		{
+			FMemory::Free((char*)InPtr - WorkaroundHeaderSize);
+		}
+	}
+
+	/** Reallocate memory returned by Alloc or Realloc */
+	static void* WorkaroundRealloc(void* InPtr, int InSizeBytes)
+	{
+		if (InPtr)
+		{
+			InPtr = (char*)InPtr - WorkaroundHeaderSize;
+		}
+		void* Result = FMemory::Realloc(InPtr, InSizeBytes + WorkaroundHeaderSize, DEFAULT_ALIGNMENT);
+		if (Result)
+		{
+			*(int*)Result = InSizeBytes;
+			Result = (char*)Result + WorkaroundHeaderSize;
+		}
+		return Result;
+	}
+
+	/** Get the actual size of an allocation returned by Alloc or Realloc */
+	static int WorkaroundSize(void* InPtr)
+	{
+		return InPtr ? *(int*)((char*)InPtr - WorkaroundHeaderSize) : 0;
+	}
+
+	/** Roundup to the expected allocation size */
+	static int WorkaroundRoundup(int InSizeBytes)
+	{
+		return (int)FMemory::QuantizeSize(InSizeBytes + WorkaroundHeaderSize, DEFAULT_ALIGNMENT) - WorkaroundHeaderSize;
+	}
 };
 
 /** Unreal implementation of an SQLite mutex */
