PR #12079: Fix memory leak when closing FSQLiteFile

diff --git a/Engine/Plugins/Runtime/Database/SQLiteCore/Source/SQLiteCore/Private/SQLiteEmbeddedPlatform.cpp b/Engine/Plugins/Runtime/Database/SQLiteCore/Source/SQLiteCore/Private/SQLiteEmbeddedPlatform.cpp
--- a/Engine/Plugins/Runtime/Database/SQLiteCore/Source/SQLiteCore/Private/SQLiteEmbeddedPlatform.cpp
+++ b/Engine/Plugins/Runtime/Database/SQLiteCore/Source/SQLiteCore/Private/SQLiteEmbeddedPlatform.cpp
@@ -433,27 +433,35 @@ private:
 		// We have to make a temporary filename if given a null filename
 		if (InFilename)
 		{
-			File->Filename = UTF8_TO_TCHAR(InFilename);
+			new (&File->Filename) FString(UTF8_TO_TCHAR(InFilename));
 		}
 		else
 		{
 			static const FString TmpPath = FPaths::ProjectIntermediateDir() / TEXT("SQLite");
 			PlatformFile.CreateDirectory(*TmpPath);
-			File->Filename = FPaths::CreateTempFilename(*TmpPath);
+			new (&File->Filename) FString(FPaths::CreateTempFilename(*TmpPath));
 		}
 
+		auto ReturnIOError = [&File]()
+		{
+			File->Filename.~FString();
+
+			return SQLITE_IOERR;
+		};
+
 		// Does the file already exist? If so "exclusive" should fail. If not, anything that doesn't specify "create" should fail
 		const bool bFileExists = PlatformFile.FileExists(*File->Filename);
 		if (bFileExists)
 		{
 			if (InFlags & SQLITE_OPEN_EXCLUSIVE)
 			{
-				return SQLITE_IOERR;
+				return ReturnIOError();
 			}
 		}
 		else if (!(InFlags & SQLITE_OPEN_CREATE))
 		{
-			return SQLITE_IOERR;
+
+			return ReturnIOError();
 		}
 
 		// Stat the file to fetch its write-ability.
@@ -487,7 +495,7 @@ private:
 		}
 		else
 		{
-			return SQLITE_IOERR;
+			return ReturnIOError();
 		}
 
 		// Opened the file - fill in the rest of the data
@@ -528,6 +536,8 @@ private:
 			PlatformFile.DeleteFile(*File->Filename);
 		}
 
+		File->Filename.~FString();
+
 		return SQLITE_OK;
 	}
 
diff --git a/Engine/Plugins/Runtime/Database/SQLiteCore/Source/SQLiteCore/Private/Tests/SQLiteTest.cpp b/Engine/Plugins/Runtime/Database/SQLiteCore/Source/SQLiteCore/Private/Tests/SQLiteTest.cpp
--- a/Engine/Plugins/Runtime/Database/SQLiteCore/Source/SQLiteCore/Private/Tests/SQLiteTest.cpp
+++ b/Engine/Plugins/Runtime/Database/SQLiteCore/Source/SQLiteCore/Private/Tests/SQLiteTest.cpp
@@ -25,6 +25,13 @@ bool FSQLiteCoreTest::RunTest(const FString& Parameters)
 	IFileManager::Get().Delete(*Path);
 	bool bSuccess = true;
 
+	// Try to open a file that doesn't exist, ensure that close cannot be called or doesn't execute unexpectedly
+	{
+		FSQLiteDatabase TestDb;
+		bSuccess &= !TestDb.Open(*Path, ESQLiteDatabaseOpenMode::ReadWrite);	// Expected to return false since the file does not exist on disk at this time
+		bSuccess &= !TestDb.Close();	// Expected to return false if nothing was opened
+	}
+
 	// Utility function to read back the records and check if the one expected.
 	auto ReadBackRecordsAndValidateValues = [](FSQLiteDatabase& Db)
 	{
