PR #12118: SQLite bugfix: Check read offset before seek

#rnx

- Our seek implementations tend to assert if the seek is invalid so we need to check for that and return an error code to avoid crashing the process.

diff --git a/Engine/Plugins/Runtime/Database/SQLiteCore/Source/SQLiteCore/Private/SQLiteEmbeddedPlatform.cpp b/Engine/Plugins/Runtime/Database/SQLiteCore/Source/SQLiteCore/Private/SQLiteEmbeddedPlatform.cpp
--- a/Engine/Plugins/Runtime/Database/SQLiteCore/Source/SQLiteCore/Private/SQLiteEmbeddedPlatform.cpp
+++ b/Engine/Plugins/Runtime/Database/SQLiteCore/Source/SQLiteCore/Private/SQLiteEmbeddedPlatform.cpp
@@ -550,6 +550,12 @@ private:
 		// Zero the buffer first in-case of a short read
 		FMemory::Memzero(OutBuffer, InReadAmountBytes);
 
+		//Seek() returns differently on different platform when offset is out of bound, so check it ahead
+		int64 iFileSize = File->FileHandle->Size();
+		if (InReadOffsetBytes >= iFileSize) {
+			return SQLITE_IOERR_SHORT_READ;
+		}
+
 		if (!File->FileHandle->Seek(InReadOffsetBytes))
 		{
 			return SQLITE_IOERR_SEEK;
