Update SQLite implementation to 3.47.1

diff --git a/Engine/Plugins/Runtime/Database/SQLiteCore/SQLiteCore.uplugin b/Engine/Plugins/Runtime/Database/SQLiteCore/SQLiteCore.uplugin
--- a/Engine/Plugins/Runtime/Database/SQLiteCore/SQLiteCore.uplugin
+++ b/Engine/Plugins/Runtime/Database/SQLiteCore/SQLiteCore.uplugin
@@ -3,7 +3,7 @@
 	"Version": 1,
 	"VersionName": "1.0",
 	"FriendlyName": "SQLite",
-	"Description": "SQLite Database",
+	"Description": "Provides a lightweight C++ wrapper for creating and manipulating SQLite databases. It uses the sqlite C library please refer to Engine/Source/ThirdParty/Licenses/SQLite_v3.47.1.license for license details.",
 	"Category": "Database",
 	"CreatedBy": "Epic Games, Inc.",
 	"CreatedByURL": "https://epicgames.com",
@@ -22,4 +22,4 @@
 			"LoadingPhase": "PreDefault"
 		}
 	]
-}
\ No newline at end of file
+}
