Update SQLite implementation to 3.47.1

diff --git a/Engine/Plugins/Runtime/Database/SQLiteCore/Source/SQLiteCore/Public/sqlite/sqlite3.h b/Engine/Plugins/Runtime/Database/SQLiteCore/Source/SQLiteCore/Public/sqlite/sqlite3.h
--- a/Engine/Plugins/Runtime/Database/SQLiteCore/Source/SQLiteCore/Public/sqlite/sqlite3.h
+++ b/Engine/Plugins/Runtime/Database/SQLiteCore/Source/SQLiteCore/Public/sqlite/sqlite3.h
@@ -43,7 +43,30 @@ extern "C" {
 
 
 /*
-** Provide the ability to override linkage features of the interface.
+** Facilitate override of interface linkage and calling conventions.
+** Be aware that these macros may not be used within this particular
+** translation of the amalgamation and its associated header file.
+**
+** The SQLITE_EXTERN and SQLITE_API macros are used to instruct the
+** compiler that the target identifier should have external linkage.
+**
+** The SQLITE_CDECL macro is used to set the calling convention for
+** public functions that accept a variable number of arguments.
+**
+** The SQLITE_APICALL macro is used to set the calling convention for
+** public functions that accept a fixed number of arguments.
+**
+** The SQLITE_STDCALL macro is no longer used and is now deprecated.
+**
+** The SQLITE_CALLBACK macro is used to set the calling convention for
+** function pointers.
+**
+** The SQLITE_SYSAPI macro is used to set the calling convention for
+** functions provided by the operating system.
+**
+** Currently, the SQLITE_CDECL, SQLITE_APICALL, SQLITE_CALLBACK, and
+** SQLITE_SYSAPI macros are used only when building for environments
+** that require non-default calling conventions.
 */
 #ifndef SQLITE_EXTERN
 # define SQLITE_EXTERN extern
@@ -108,7 +131,7 @@ extern "C" {
 ** be held constant and Z will be incremented or else Y will be incremented
 ** and Z will be reset to zero.
 **
-** Since [version 3.6.18] ([dateof:3.6.18]), 
+** Since [version 3.6.18] ([dateof:3.6.18]),
 ** SQLite source code has been stored in the
 ** <a href="http://www.fossil-scm.org/">Fossil configuration management
 ** system</a>.  ^The SQLITE_SOURCE_ID macro evaluates to
@@ -123,9 +146,9 @@ extern "C" {
 ** [sqlite3_libversion_number()], [sqlite3_sourceid()],
 ** [sqlite_version()] and [sqlite_source_id()].
 */
-#define SQLITE_VERSION        "3.31.1"
-#define SQLITE_VERSION_NUMBER 3031001
-#define SQLITE_SOURCE_ID      "2020-01-27 19:55:54 3bfa9cc97da10598521b342961df8f5f68c7388fa117345eeb516eaa837bb4d6"
+#define SQLITE_VERSION        "3.47.1"
+#define SQLITE_VERSION_NUMBER 3047001
+#define SQLITE_SOURCE_ID      "2024-11-25 12:07:48 b95d11e958643b969c47a8e5857f3793b9e69700b8f1469371386369a26e577e"
 
 /*
 ** CAPI3REF: Run-Time Library Version Numbers
@@ -151,8 +174,8 @@ extern "C" {
 ** function is provided for use in DLLs since DLL users usually do not have
 ** direct access to string constants within the DLL.  ^The
 ** sqlite3_libversion_number() function returns an integer equal to
-** [SQLITE_VERSION_NUMBER].  ^(The sqlite3_sourceid() function returns 
-** a pointer to a string constant whose value is the same as the 
+** [SQLITE_VERSION_NUMBER].  ^(The sqlite3_sourceid() function returns
+** a pointer to a string constant whose value is the same as the
 ** [SQLITE_SOURCE_ID] C preprocessor macro.  Except if SQLite is built
 ** using an edited copy of [the amalgamation], then the last four characters
 ** of the hash might be different from [SQLITE_SOURCE_ID].)^
@@ -167,20 +190,20 @@ SQLITE_API int sqlite3_libversion_number(void);
 /*
 ** CAPI3REF: Run-Time Library Compilation Options Diagnostics
 **
-** ^The sqlite3_compileoption_used() function returns 0 or 1 
-** indicating whether the specified option was defined at 
-** compile time.  ^The SQLITE_ prefix may be omitted from the 
-** option name passed to sqlite3_compileoption_used().  
+** ^The sqlite3_compileoption_used() function returns 0 or 1
+** indicating whether the specified option was defined at
+** compile time.  ^The SQLITE_ prefix may be omitted from the
+** option name passed to sqlite3_compileoption_used().
 **
 ** ^The sqlite3_compileoption_get() function allows iterating
 ** over the list of options that were defined at compile time by
 ** returning the N-th compile time option string.  ^If N is out of range,
-** sqlite3_compileoption_get() returns a NULL pointer.  ^The SQLITE_ 
-** prefix is omitted from any strings returned by 
+** sqlite3_compileoption_get() returns a NULL pointer.  ^The SQLITE_
+** prefix is omitted from any strings returned by
 ** sqlite3_compileoption_get().
 **
 ** ^Support for the diagnostic functions sqlite3_compileoption_used()
-** and sqlite3_compileoption_get() may be omitted by specifying the 
+** and sqlite3_compileoption_get() may be omitted by specifying the
 ** [SQLITE_OMIT_COMPILEOPTION_DIAGS] option at compile time.
 **
 ** See also: SQL functions [sqlite_compileoption_used()] and
@@ -204,7 +227,7 @@ SQLITE_API const char *sqlite3_compileoption_get(int N);
 ** SQLite can be compiled with or without mutexes.  When
 ** the [SQLITE_THREADSAFE] C preprocessor macro is 1 or 2, mutexes
 ** are enabled and SQLite is threadsafe.  When the
-** [SQLITE_THREADSAFE] macro is 0, 
+** [SQLITE_THREADSAFE] macro is 0,
 ** the mutexes are omitted.  Without the mutexes, it is not safe
 ** to use SQLite concurrently from more than one thread.
 **
@@ -261,14 +284,14 @@ typedef struct sqlite3 sqlite3;
 **
 ** ^The sqlite3_int64 and sqlite_int64 types can store integer values
 ** between -9223372036854775808 and +9223372036854775807 inclusive.  ^The
-** sqlite3_uint64 and sqlite_uint64 types can store integer values 
+** sqlite3_uint64 and sqlite_uint64 types can store integer values
 ** between 0 and +18446744073709551615 inclusive.
 */
 #ifdef SQLITE_INT64_TYPE
   typedef SQLITE_INT64_TYPE sqlite_int64;
 # ifdef SQLITE_UINT64_TYPE
     typedef SQLITE_UINT64_TYPE sqlite_uint64;
-# else  
+# else
     typedef unsigned SQLITE_INT64_TYPE sqlite_uint64;
 # endif
 #elif defined(_MSC_VER) || defined(__BORLANDC__)
@@ -299,26 +322,22 @@ typedef sqlite_uint64 sqlite3_uint64;
 ** the [sqlite3] object is successfully destroyed and all associated
 ** resources are deallocated.
 **
-** ^If the database connection is associated with unfinalized prepared
-** statements or unfinished sqlite3_backup objects then sqlite3_close()
-** will leave the database connection open and return [SQLITE_BUSY].
-** ^If sqlite3_close_v2() is called with unfinalized prepared statements
-** and/or unfinished sqlite3_backups, then the database connection becomes
-** an unusable "zombie" which will automatically be deallocated when the
-** last prepared statement is finalized or the last sqlite3_backup is
-** finished.  The sqlite3_close_v2() interface is intended for use with
-** host languages that are garbage collected, and where the order in which
-** destructors are called is arbitrary.
-**
-** Applications should [sqlite3_finalize | finalize] all [prepared statements],
-** [sqlite3_blob_close | close] all [BLOB handles], and 
+** Ideally, applications should [sqlite3_finalize | finalize] all
+** [prepared statements], [sqlite3_blob_close | close] all [BLOB handles], and
 ** [sqlite3_backup_finish | finish] all [sqlite3_backup] objects associated
-** with the [sqlite3] object prior to attempting to close the object.  ^If
-** sqlite3_close_v2() is called on a [database connection] that still has
-** outstanding [prepared statements], [BLOB handles], and/or
-** [sqlite3_backup] objects then it returns [SQLITE_OK] and the deallocation
-** of resources is deferred until all [prepared statements], [BLOB handles],
-** and [sqlite3_backup] objects are also destroyed.
+** with the [sqlite3] object prior to attempting to close the object.
+** ^If the database connection is associated with unfinalized prepared
+** statements, BLOB handlers, and/or unfinished sqlite3_backup objects then
+** sqlite3_close() will leave the database connection open and return
+** [SQLITE_BUSY]. ^If sqlite3_close_v2() is called with unfinalized prepared
+** statements, unclosed BLOB handlers, and/or unfinished sqlite3_backups,
+** it returns [SQLITE_OK] regardless, but instead of deallocating the database
+** connection immediately, it marks the database connection as an unusable
+** "zombie" and makes arrangements to automatically deallocate the database
+** connection after all prepared statements are finalized, all BLOB handles
+** are closed, and all backups have finished. The sqlite3_close_v2() interface
+** is intended for use with host languages that are garbage collected, and
+** where the order in which destructors are called is arbitrary.
 **
 ** ^If an [sqlite3] object is destroyed while a transaction is open,
 ** the transaction is automatically rolled back.
@@ -348,7 +367,7 @@ typedef int (*sqlite3_callback)(void*,int,char**, char**);
 ** The sqlite3_exec() interface is a convenience wrapper around
 ** [sqlite3_prepare_v2()], [sqlite3_step()], and [sqlite3_finalize()],
 ** that allows an application to run multiple statements of SQL
-** without having to use a lot of C code. 
+** without having to use a lot of C code.
 **
 ** ^The sqlite3_exec() interface runs zero or more UTF-8 encoded,
 ** semicolon-separate SQL statements passed into its 2nd argument,
@@ -388,7 +407,7 @@ typedef int (*sqlite3_callback)(void*,int,char**, char**);
 ** from [sqlite3_column_name()].
 **
 ** ^If the 2nd parameter to sqlite3_exec() is a NULL pointer, a pointer
-** to an empty string, or a pointer that contains only whitespace and/or 
+** to an empty string, or a pointer that contains only whitespace and/or
 ** SQL comments, then no SQL statements are evaluated and the database
 ** is not changed.
 **
@@ -401,6 +420,8 @@ typedef int (*sqlite3_callback)(void*,int,char**, char**);
 **      the 1st parameter to sqlite3_exec() while sqlite3_exec() is running.
 ** <li> The application must not modify the SQL statement text passed into
 **      the 2nd parameter of sqlite3_exec() while sqlite3_exec() is running.
+** <li> The application must not dereference the arrays or string pointers
+**       passed as the 3rd and 4th callback parameters after it returns.
 ** </ul>
 */
 SQLITE_API int sqlite3_exec(
@@ -507,10 +528,14 @@ SQLITE_API int sqlite3_exec(
 #define SQLITE_IOERR_BEGIN_ATOMIC      (SQLITE_IOERR | (29<<8))
 #define SQLITE_IOERR_COMMIT_ATOMIC     (SQLITE_IOERR | (30<<8))
 #define SQLITE_IOERR_ROLLBACK_ATOMIC   (SQLITE_IOERR | (31<<8))
+#define SQLITE_IOERR_DATA              (SQLITE_IOERR | (32<<8))
+#define SQLITE_IOERR_CORRUPTFS         (SQLITE_IOERR | (33<<8))
+#define SQLITE_IOERR_IN_PAGE           (SQLITE_IOERR | (34<<8))
 #define SQLITE_LOCKED_SHAREDCACHE      (SQLITE_LOCKED |  (1<<8))
 #define SQLITE_LOCKED_VTAB             (SQLITE_LOCKED |  (2<<8))
 #define SQLITE_BUSY_RECOVERY           (SQLITE_BUSY   |  (1<<8))
 #define SQLITE_BUSY_SNAPSHOT           (SQLITE_BUSY   |  (2<<8))
+#define SQLITE_BUSY_TIMEOUT            (SQLITE_BUSY   |  (3<<8))
 #define SQLITE_CANTOPEN_NOTEMPDIR      (SQLITE_CANTOPEN | (1<<8))
 #define SQLITE_CANTOPEN_ISDIR          (SQLITE_CANTOPEN | (2<<8))
 #define SQLITE_CANTOPEN_FULLPATH       (SQLITE_CANTOPEN | (3<<8))
@@ -519,6 +544,7 @@ SQLITE_API int sqlite3_exec(
 #define SQLITE_CANTOPEN_SYMLINK        (SQLITE_CANTOPEN | (6<<8))
 #define SQLITE_CORRUPT_VTAB            (SQLITE_CORRUPT | (1<<8))
 #define SQLITE_CORRUPT_SEQUENCE        (SQLITE_CORRUPT | (2<<8))
+#define SQLITE_CORRUPT_INDEX           (SQLITE_CORRUPT | (3<<8))
 #define SQLITE_READONLY_RECOVERY       (SQLITE_READONLY | (1<<8))
 #define SQLITE_READONLY_CANTLOCK       (SQLITE_READONLY | (2<<8))
 #define SQLITE_READONLY_ROLLBACK       (SQLITE_READONLY | (3<<8))
@@ -537,12 +563,14 @@ SQLITE_API int sqlite3_exec(
 #define SQLITE_CONSTRAINT_VTAB         (SQLITE_CONSTRAINT | (9<<8))
 #define SQLITE_CONSTRAINT_ROWID        (SQLITE_CONSTRAINT |(10<<8))
 #define SQLITE_CONSTRAINT_PINNED       (SQLITE_CONSTRAINT |(11<<8))
+#define SQLITE_CONSTRAINT_DATATYPE     (SQLITE_CONSTRAINT |(12<<8))
 #define SQLITE_NOTICE_RECOVER_WAL      (SQLITE_NOTICE | (1<<8))
 #define SQLITE_NOTICE_RECOVER_ROLLBACK (SQLITE_NOTICE | (2<<8))
+#define SQLITE_NOTICE_RBU              (SQLITE_NOTICE | (3<<8))
 #define SQLITE_WARNING_AUTOINDEX       (SQLITE_WARNING | (1<<8))
 #define SQLITE_AUTH_USER               (SQLITE_AUTH | (1<<8))
 #define SQLITE_OK_LOAD_PERMANENTLY     (SQLITE_OK | (1<<8))
-#define SQLITE_OK_SYMLINK              (SQLITE_OK | (2<<8))
+#define SQLITE_OK_SYMLINK              (SQLITE_OK | (2<<8)) /* internal use only */
 
 /*
 ** CAPI3REF: Flags For File Open Operations
@@ -550,6 +578,19 @@ SQLITE_API int sqlite3_exec(
 ** These bit values are intended for use in the
 ** 3rd parameter to the [sqlite3_open_v2()] interface and
 ** in the 4th parameter to the [sqlite3_vfs.xOpen] method.
+**
+** Only those flags marked as "Ok for sqlite3_open_v2()" may be
+** used as the third argument to the [sqlite3_open_v2()] interface.
+** The other flags have historically been ignored by sqlite3_open_v2(),
+** though future versions of SQLite might change so that an error is
+** raised if any of the disallowed bits are passed into sqlite3_open_v2().
+** Applications should not depend on the historical behavior.
+**
+** Note in particular that passing the SQLITE_OPEN_EXCLUSIVE flag into
+** [sqlite3_open_v2()] does *not* cause the underlying database file
+** to be opened using O_EXCL.  Passing SQLITE_OPEN_EXCLUSIVE into
+** [sqlite3_open_v2()] has historically be a no-op and might become an
+** error in future versions of SQLite.
 */
 #define SQLITE_OPEN_READONLY         0x00000001  /* Ok for sqlite3_open_v2() */
 #define SQLITE_OPEN_READWRITE        0x00000002  /* Ok for sqlite3_open_v2() */
@@ -565,15 +606,19 @@ SQLITE_API int sqlite3_exec(
 #define SQLITE_OPEN_MAIN_JOURNAL     0x00000800  /* VFS only */
 #define SQLITE_OPEN_TEMP_JOURNAL     0x00001000  /* VFS only */
 #define SQLITE_OPEN_SUBJOURNAL       0x00002000  /* VFS only */
-#define SQLITE_OPEN_MASTER_JOURNAL   0x00004000  /* VFS only */
+#define SQLITE_OPEN_SUPER_JOURNAL    0x00004000  /* VFS only */
 #define SQLITE_OPEN_NOMUTEX          0x00008000  /* Ok for sqlite3_open_v2() */
 #define SQLITE_OPEN_FULLMUTEX        0x00010000  /* Ok for sqlite3_open_v2() */
 #define SQLITE_OPEN_SHAREDCACHE      0x00020000  /* Ok for sqlite3_open_v2() */
 #define SQLITE_OPEN_PRIVATECACHE     0x00040000  /* Ok for sqlite3_open_v2() */
 #define SQLITE_OPEN_WAL              0x00080000  /* VFS only */
 #define SQLITE_OPEN_NOFOLLOW         0x01000000  /* Ok for sqlite3_open_v2() */
+#define SQLITE_OPEN_EXRESCODE        0x02000000  /* Extended result codes */
 
 /* Reserved:                         0x00F00000 */
+/* Legacy compatibility: */
+#define SQLITE_OPEN_MASTER_JOURNAL   0x00004000  /* VFS only */
+
 
 /*
 ** CAPI3REF: Device Characteristics
@@ -607,6 +652,13 @@ SQLITE_API int sqlite3_exec(
 ** filesystem supports doing multiple write operations atomically when those
 ** write operations are bracketed by [SQLITE_FCNTL_BEGIN_ATOMIC_WRITE] and
 ** [SQLITE_FCNTL_COMMIT_ATOMIC_WRITE].
+**
+** The SQLITE_IOCAP_SUBPAGE_READ property means that it is ok to read
+** from the database file in amounts that are not a multiple of the
+** page size and that do not begin at a page boundary.  Without this
+** property, SQLite is careful to only do full-page reads and write
+** on aligned pages, with the one exception that it will do a sub-page
+** read of the first page to access the database header.
 */
 #define SQLITE_IOCAP_ATOMIC                 0x00000001
 #define SQLITE_IOCAP_ATOMIC512              0x00000002
@@ -623,19 +675,24 @@ SQLITE_API int sqlite3_exec(
 #define SQLITE_IOCAP_POWERSAFE_OVERWRITE    0x00001000
 #define SQLITE_IOCAP_IMMUTABLE              0x00002000
 #define SQLITE_IOCAP_BATCH_ATOMIC           0x00004000
+#define SQLITE_IOCAP_SUBPAGE_READ           0x00008000
 
 /*
 ** CAPI3REF: File Locking Levels
 **
 ** SQLite uses one of these integer values as the second
 ** argument to calls it makes to the xLock() and xUnlock() methods
-** of an [sqlite3_io_methods] object.
+** of an [sqlite3_io_methods] object.  These values are ordered from
+** lest restrictive to most restrictive.
+**
+** The argument to xLock() is always SHARED or higher.  The argument to
+** xUnlock is either SHARED or NONE.
 */
-#define SQLITE_LOCK_NONE          0
-#define SQLITE_LOCK_SHARED        1
-#define SQLITE_LOCK_RESERVED      2
-#define SQLITE_LOCK_PENDING       3
-#define SQLITE_LOCK_EXCLUSIVE     4
+#define SQLITE_LOCK_NONE          0       /* xUnlock() only */
+#define SQLITE_LOCK_SHARED        1       /* xLock() or xUnlock() */
+#define SQLITE_LOCK_RESERVED      2       /* xLock() only */
+#define SQLITE_LOCK_PENDING       3       /* xLock() only */
+#define SQLITE_LOCK_EXCLUSIVE     4       /* xLock() only */
 
 /*
 ** CAPI3REF: Synchronization Type Flags
@@ -670,7 +727,7 @@ SQLITE_API int sqlite3_exec(
 /*
 ** CAPI3REF: OS Interface Open File Handle
 **
-** An [sqlite3_file] object represents an open file in the 
+** An [sqlite3_file] object represents an open file in the
 ** [sqlite3_vfs | OS interface layer].  Individual OS interface
 ** implementations will
 ** want to subclass this object by appending additional fields
@@ -692,7 +749,7 @@ struct sqlite3_file {
 ** This object defines the methods used to perform various operations
 ** against the open file represented by the [sqlite3_file] object.
 **
-** If the [sqlite3_vfs.xOpen] method sets the sqlite3_file.pMethods element 
+** If the [sqlite3_vfs.xOpen] method sets the sqlite3_file.pMethods element
 ** to a non-NULL pointer, then the sqlite3_io_methods.xClose method
 ** may be invoked even if the [sqlite3_vfs.xOpen] reported that it failed.  The
 ** only way to prevent a call to xClose following a failed [sqlite3_vfs.xOpen]
@@ -713,11 +770,18 @@ struct sqlite3_file {
 ** <li> [SQLITE_LOCK_PENDING], or
 ** <li> [SQLITE_LOCK_EXCLUSIVE].
 ** </ul>
-** xLock() increases the lock. xUnlock() decreases the lock.
+** xLock() upgrades the database file lock.  In other words, xLock() moves the
+** database file lock in the direction NONE toward EXCLUSIVE. The argument to
+** xLock() is always one of SHARED, RESERVED, PENDING, or EXCLUSIVE, never
+** SQLITE_LOCK_NONE.  If the database file lock is already at or above the
+** requested lock, then the call to xLock() is a no-op.
+** xUnlock() downgrades the database file lock to either SHARED or NONE.
+** If the lock is already at or below the requested lock state, then the call
+** to xUnlock() is a no-op.
 ** The xCheckReservedLock() method checks whether any database connection,
 ** either in this process or in some other process, is holding a RESERVED,
-** PENDING, or EXCLUSIVE lock on the file.  It returns true
-** if such a lock exists and false otherwise.
+** PENDING, or EXCLUSIVE lock on the file.  It returns, via its output
+** pointer parameter, true if such a lock exists and false otherwise.
 **
 ** The xFileControl() method is a generic interface that allows custom
 ** VFS implementations to directly control an open file using the
@@ -758,6 +822,7 @@ struct sqlite3_file {
 ** <li> [SQLITE_IOCAP_POWERSAFE_OVERWRITE]
 ** <li> [SQLITE_IOCAP_IMMUTABLE]
 ** <li> [SQLITE_IOCAP_BATCH_ATOMIC]
+** <li> [SQLITE_IOCAP_SUBPAGE_READ]
 ** </ul>
 **
 ** The SQLITE_IOCAP_ATOMIC property means that all writes of
@@ -818,9 +883,8 @@ struct sqlite3_io_methods {
 ** opcode causes the xFileControl method to write the current state of
 ** the lock (one of [SQLITE_LOCK_NONE], [SQLITE_LOCK_SHARED],
 ** [SQLITE_LOCK_RESERVED], [SQLITE_LOCK_PENDING], or [SQLITE_LOCK_EXCLUSIVE])
-** into an integer that the pArg argument points to. This capability
-** is used during testing and is only available when the SQLITE_TEST
-** compile-time option is used.
+** into an integer that the pArg argument points to.
+** This capability is only available if SQLite is compiled with [SQLITE_DEBUG].
 **
 ** <li>[[SQLITE_FCNTL_SIZE_HINT]]
 ** The [SQLITE_FCNTL_SIZE_HINT] opcode is used by SQLite to give the VFS
@@ -842,7 +906,7 @@ struct sqlite3_io_methods {
 ** <li>[[SQLITE_FCNTL_CHUNK_SIZE]]
 ** The [SQLITE_FCNTL_CHUNK_SIZE] opcode is used to request that the VFS
 ** extends and truncates the database file in chunks of a size specified
-** by the user. The fourth argument to [sqlite3_file_control()] should 
+** by the user. The fourth argument to [sqlite3_file_control()] should
 ** point to an integer (type int) containing the new chunk-size to use
 ** for the nominated database. Allocating database file space in large
 ** chunks (say 1MB at a time), may reduce file-system fragmentation and
@@ -865,24 +929,24 @@ struct sqlite3_io_methods {
 ** <li>[[SQLITE_FCNTL_SYNC]]
 ** The [SQLITE_FCNTL_SYNC] opcode is generated internally by SQLite and
 ** sent to the VFS immediately before the xSync method is invoked on a
-** database file descriptor. Or, if the xSync method is not invoked 
-** because the user has configured SQLite with 
-** [PRAGMA synchronous | PRAGMA synchronous=OFF] it is invoked in place 
+** database file descriptor. Or, if the xSync method is not invoked
+** because the user has configured SQLite with
+** [PRAGMA synchronous | PRAGMA synchronous=OFF] it is invoked in place
 ** of the xSync method. In most cases, the pointer argument passed with
 ** this file-control is NULL. However, if the database file is being synced
 ** as part of a multi-database commit, the argument points to a nul-terminated
-** string containing the transactions master-journal file name. VFSes that 
-** do not need this signal should silently ignore this opcode. Applications 
-** should not call [sqlite3_file_control()] with this opcode as doing so may 
-** disrupt the operation of the specialized VFSes that do require it.  
+** string containing the transactions super-journal file name. VFSes that
+** do not need this signal should silently ignore this opcode. Applications
+** should not call [sqlite3_file_control()] with this opcode as doing so may
+** disrupt the operation of the specialized VFSes that do require it.
 **
 ** <li>[[SQLITE_FCNTL_COMMIT_PHASETWO]]
 ** The [SQLITE_FCNTL_COMMIT_PHASETWO] opcode is generated internally by SQLite
 ** and sent to the VFS after a transaction has been committed immediately
 ** but before the database is unlocked. VFSes that do not need this signal
 ** should silently ignore this opcode. Applications should not call
-** [sqlite3_file_control()] with this opcode as doing so may disrupt the 
-** operation of the specialized VFSes that do require it.  
+** [sqlite3_file_control()] with this opcode as doing so may disrupt the
+** operation of the specialized VFSes that do require it.
 **
 ** <li>[[SQLITE_FCNTL_WIN32_AV_RETRY]]
 ** ^The [SQLITE_FCNTL_WIN32_AV_RETRY] opcode is used to configure automatic
@@ -930,13 +994,13 @@ struct sqlite3_io_methods {
 ** <li>[[SQLITE_FCNTL_OVERWRITE]]
 ** ^The [SQLITE_FCNTL_OVERWRITE] opcode is invoked by SQLite after opening
 ** a write transaction to indicate that, unless it is rolled back for some
-** reason, the entire database file will be overwritten by the current 
+** reason, the entire database file will be overwritten by the current
 ** transaction. This is used by VACUUM operations.
 **
 ** <li>[[SQLITE_FCNTL_VFSNAME]]
 ** ^The [SQLITE_FCNTL_VFSNAME] opcode can be used to obtain the names of
 ** all [VFSes] in the VFS stack.  The names are of all VFS shims and the
-** final bottom-level VFS are written into memory obtained from 
+** final bottom-level VFS are written into memory obtained from
 ** [sqlite3_malloc()] and the result is stored in the char* variable
 ** that the fourth parameter of [sqlite3_file_control()] points to.
 ** The caller is responsible for freeing the memory when done.  As with
@@ -955,7 +1019,7 @@ struct sqlite3_io_methods {
 ** upper-most shim only.
 **
 ** <li>[[SQLITE_FCNTL_PRAGMA]]
-** ^Whenever a [PRAGMA] statement is parsed, an [SQLITE_FCNTL_PRAGMA] 
+** ^Whenever a [PRAGMA] statement is parsed, an [SQLITE_FCNTL_PRAGMA]
 ** file control is sent to the open [sqlite3_file] object corresponding
 ** to the database file to which the pragma statement refers. ^The argument
 ** to the [SQLITE_FCNTL_PRAGMA] file control is an array of
@@ -966,7 +1030,7 @@ struct sqlite3_io_methods {
 ** of the char** argument point to a string obtained from [sqlite3_mprintf()]
 ** or the equivalent and that string will become the result of the pragma or
 ** the error message if the pragma fails. ^If the
-** [SQLITE_FCNTL_PRAGMA] file control returns [SQLITE_NOTFOUND], then normal 
+** [SQLITE_FCNTL_PRAGMA] file control returns [SQLITE_NOTFOUND], then normal
 ** [PRAGMA] processing continues.  ^If the [SQLITE_FCNTL_PRAGMA]
 ** file control returns [SQLITE_OK], then the parser assumes that the
 ** VFS has handled the PRAGMA itself and the parser generates a no-op
@@ -1006,7 +1070,7 @@ struct sqlite3_io_methods {
 ** The argument is a pointer to a value of type sqlite3_int64 that
 ** is an advisory maximum number of bytes in the file to memory map.  The
 ** pointer is overwritten with the old value.  The limit is not changed if
-** the value originally pointed to is negative, and so the current limit 
+** the value originally pointed to is negative, and so the current limit
 ** can be queried by passing in a pointer to a negative number.  This
 ** file-control is used internally to implement [PRAGMA mmap_size].
 **
@@ -1050,7 +1114,7 @@ struct sqlite3_io_methods {
 ** <li>[[SQLITE_FCNTL_RBU]]
 ** The [SQLITE_FCNTL_RBU] opcode is implemented by the special VFS used by
 ** the RBU extension only.  All other VFS should return SQLITE_NOTFOUND for
-** this opcode.  
+** this opcode.
 **
 ** <li>[[SQLITE_FCNTL_BEGIN_ATOMIC_WRITE]]
 ** If the [SQLITE_FCNTL_BEGIN_ATOMIC_WRITE] opcode returns SQLITE_OK, then
@@ -1067,7 +1131,7 @@ struct sqlite3_io_methods {
 **
 ** <li>[[SQLITE_FCNTL_COMMIT_ATOMIC_WRITE]]
 ** The [SQLITE_FCNTL_COMMIT_ATOMIC_WRITE] opcode causes all write
-** operations since the previous successful call to 
+** operations since the previous successful call to
 ** [SQLITE_FCNTL_BEGIN_ATOMIC_WRITE] to be performed atomically.
 ** This file control returns [SQLITE_OK] if and only if the writes were
 ** all performed successfully and have been committed to persistent storage.
@@ -1079,7 +1143,7 @@ struct sqlite3_io_methods {
 **
 ** <li>[[SQLITE_FCNTL_ROLLBACK_ATOMIC_WRITE]]
 ** The [SQLITE_FCNTL_ROLLBACK_ATOMIC_WRITE] opcode causes all write
-** operations since the previous successful call to 
+** operations since the previous successful call to
 ** [SQLITE_FCNTL_BEGIN_ATOMIC_WRITE] to be rolled back.
 ** ^This file control takes the file descriptor out of batch write mode
 ** so that all subsequent write operations are independent.
@@ -1087,10 +1151,12 @@ struct sqlite3_io_methods {
 ** a prior successful call to [SQLITE_FCNTL_BEGIN_ATOMIC_WRITE].
 **
 ** <li>[[SQLITE_FCNTL_LOCK_TIMEOUT]]
-** The [SQLITE_FCNTL_LOCK_TIMEOUT] opcode causes attempts to obtain
-** a file lock using the xLock or xShmLock methods of the VFS to wait
-** for up to M milliseconds before failing, where M is the single 
-** unsigned integer parameter.
+** The [SQLITE_FCNTL_LOCK_TIMEOUT] opcode is used to configure a VFS
+** to block for up to M milliseconds before failing when attempting to
+** obtain a file lock using the xLock or xShmLock methods of the VFS.
+** The parameter is a pointer to a 32-bit signed integer that contains
+** the value that M is to be set to. Before returning, the 32-bit signed
+** integer is overwritten with the previous value of M.
 **
 ** <li>[[SQLITE_FCNTL_DATA_VERSION]]
 ** The [SQLITE_FCNTL_DATA_VERSION] opcode is used to detect changes to
@@ -1112,11 +1178,38 @@ struct sqlite3_io_methods {
 ** happen either internally or externally and that are associated with
 ** a particular attached database.
 **
+** <li>[[SQLITE_FCNTL_CKPT_START]]
+** The [SQLITE_FCNTL_CKPT_START] opcode is invoked from within a checkpoint
+** in wal mode before the client starts to copy pages from the wal
+** file to the database file.
+**
 ** <li>[[SQLITE_FCNTL_CKPT_DONE]]
 ** The [SQLITE_FCNTL_CKPT_DONE] opcode is invoked from within a checkpoint
 ** in wal mode after the client has finished copying pages from the wal
 ** file to the database file, but before the *-shm file is updated to
 ** record the fact that the pages have been checkpointed.
+**
+** <li>[[SQLITE_FCNTL_EXTERNAL_READER]]
+** The EXPERIMENTAL [SQLITE_FCNTL_EXTERNAL_READER] opcode is used to detect
+** whether or not there is a database client in another process with a wal-mode
+** transaction open on the database or not. It is only available on unix.The
+** (void*) argument passed with this file-control should be a pointer to a
+** value of type (int). The integer value is set to 1 if the database is a wal
+** mode database and there exists at least one client in another process that
+** currently has an SQL transaction open on the database. It is set to 0 if
+** the database is not a wal-mode db, or if there is no such connection in any
+** other process. This opcode cannot be used to detect transactions opened
+** by clients within the current process, only within other processes.
+**
+** <li>[[SQLITE_FCNTL_CKSM_FILE]]
+** The [SQLITE_FCNTL_CKSM_FILE] opcode is for use internally by the
+** [checksum VFS shim] only.
+**
+** <li>[[SQLITE_FCNTL_RESET_CACHE]]
+** If there is currently no transaction open on the database, and the
+** database is not a temp db, then the [SQLITE_FCNTL_RESET_CACHE] file-control
+** purges the contents of the in-memory page cache. If there is an open
+** transaction, or if the db is a temp-db, this opcode is a no-op, not an error.
 ** </ul>
 */
 #define SQLITE_FCNTL_LOCKSTATE               1
@@ -1155,6 +1248,11 @@ struct sqlite3_io_methods {
 #define SQLITE_FCNTL_DATA_VERSION           35
 #define SQLITE_FCNTL_SIZE_LIMIT             36
 #define SQLITE_FCNTL_CKPT_DONE              37
+#define SQLITE_FCNTL_RESERVE_BYTES          38
+#define SQLITE_FCNTL_CKPT_START             39
+#define SQLITE_FCNTL_EXTERNAL_READER        40
+#define SQLITE_FCNTL_CKSM_FILE              41
+#define SQLITE_FCNTL_RESET_CACHE            42
 
 /* deprecated names */
 #define SQLITE_GET_LOCKPROXYFILE      SQLITE_FCNTL_GET_LOCKPROXYFILE
@@ -1184,6 +1282,26 @@ typedef struct sqlite3_mutex sqlite3_mutex;
 */
 typedef struct sqlite3_api_routines sqlite3_api_routines;
 
+/*
+** CAPI3REF: File Name
+**
+** Type [sqlite3_filename] is used by SQLite to pass filenames to the
+** xOpen method of a [VFS]. It may be cast to (const char*) and treated
+** as a normal, nul-terminated, UTF-8 buffer containing the filename, but
+** may also be passed to special APIs such as:
+**
+** <ul>
+** <li>  sqlite3_filename_database()
+** <li>  sqlite3_filename_journal()
+** <li>  sqlite3_filename_wal()
+** <li>  sqlite3_uri_parameter()
+** <li>  sqlite3_uri_boolean()
+** <li>  sqlite3_uri_int64()
+** <li>  sqlite3_uri_key()
+** </ul>
+*/
+typedef const char *sqlite3_filename;
+
 /*
 ** CAPI3REF: OS Interface Object
 **
@@ -1238,14 +1356,14 @@ typedef struct sqlite3_api_routines sqlite3_api_routines;
 ** the [sqlite3_file] can safely store a pointer to the
 ** filename if it needs to remember the filename for some reason.
 ** If the zFilename parameter to xOpen is a NULL pointer then xOpen
-** must invent its own temporary name for the file.  ^Whenever the 
+** must invent its own temporary name for the file.  ^Whenever the
 ** xFilename parameter is NULL it will also be the case that the
 ** flags parameter will include [SQLITE_OPEN_DELETEONCLOSE].
 **
 ** The flags argument to xOpen() includes all bits set in
 ** the flags argument to [sqlite3_open_v2()].  Or if [sqlite3_open()]
 ** or [sqlite3_open16()] is used, then flags includes at least
-** [SQLITE_OPEN_READWRITE] | [SQLITE_OPEN_CREATE]. 
+** [SQLITE_OPEN_READWRITE] | [SQLITE_OPEN_CREATE].
 ** If xOpen() opens a file read-only then it sets *pOutFlags to
 ** include [SQLITE_OPEN_READONLY].  Other bits in *pOutFlags may be set.
 **
@@ -1259,7 +1377,7 @@ typedef struct sqlite3_api_routines sqlite3_api_routines;
 ** <li>  [SQLITE_OPEN_TEMP_JOURNAL]
 ** <li>  [SQLITE_OPEN_TRANSIENT_DB]
 ** <li>  [SQLITE_OPEN_SUBJOURNAL]
-** <li>  [SQLITE_OPEN_MASTER_JOURNAL]
+** <li>  [SQLITE_OPEN_SUPER_JOURNAL]
 ** <li>  [SQLITE_OPEN_WAL]
 ** </ul>)^
 **
@@ -1287,10 +1405,10 @@ typedef struct sqlite3_api_routines sqlite3_api_routines;
 ** ^The [SQLITE_OPEN_EXCLUSIVE] flag is always used in conjunction
 ** with the [SQLITE_OPEN_CREATE] flag, which are both directly
 ** analogous to the O_EXCL and O_CREAT flags of the POSIX open()
-** API.  The SQLITE_OPEN_EXCLUSIVE flag, when paired with the 
+** API.  The SQLITE_OPEN_EXCLUSIVE flag, when paired with the
 ** SQLITE_OPEN_CREATE, is used to indicate that file should always
 ** be created, and that it is an error if it already exists.
-** It is <i>not</i> used to indicate the file should be opened 
+** It is <i>not</i> used to indicate the file should be opened
 ** for exclusive access.
 **
 ** ^At least szOsFile bytes of memory are allocated by SQLite
@@ -1314,7 +1432,7 @@ typedef struct sqlite3_api_routines sqlite3_api_routines;
 ** non-zero error code if there is an I/O error or if the name of
 ** the file given in the second argument is illegal.  If SQLITE_OK
 ** is returned, then non-zero or zero is written into *pResOut to indicate
-** whether or not the file is accessible.  
+** whether or not the file is accessible.
 **
 ** ^SQLite will always allocate at least mxPathname+1 bytes for the
 ** output buffer xFullPathname.  The exact size of the output buffer
@@ -1334,16 +1452,16 @@ typedef struct sqlite3_api_routines sqlite3_api_routines;
 ** method returns a Julian Day Number for the current date and time as
 ** a floating point value.
 ** ^The xCurrentTimeInt64() method returns, as an integer, the Julian
-** Day Number multiplied by 86400000 (the number of milliseconds in 
-** a 24-hour day).  
+** Day Number multiplied by 86400000 (the number of milliseconds in
+** a 24-hour day).
 ** ^SQLite will use the xCurrentTimeInt64() method to get the current
-** date and time if that method is available (if iVersion is 2 or 
+** date and time if that method is available (if iVersion is 2 or
 ** greater and the function pointer is not NULL) and will fall back
 ** to xCurrentTime() if xCurrentTimeInt64() is unavailable.
 **
 ** ^The xSetSystemCall(), xGetSystemCall(), and xNestSystemCall() interfaces
 ** are not used by the SQLite core.  These optional interfaces are provided
-** by some VFSes to facilitate testing of the VFS code. By overriding 
+** by some VFSes to facilitate testing of the VFS code. By overriding
 ** system calls with functions under its control, a test program can
 ** simulate faults and error conditions that would otherwise be difficult
 ** or impossible to induce.  The set of system calls that can be overridden
@@ -1362,7 +1480,7 @@ struct sqlite3_vfs {
   sqlite3_vfs *pNext;      /* Next registered VFS */
   const char *zName;       /* Name of this virtual file system */
   void *pAppData;          /* Pointer to application-specific data */
-  int (*xOpen)(sqlite3_vfs*, const char *zName, sqlite3_file*,
+  int (*xOpen)(sqlite3_vfs*, sqlite3_filename zName, sqlite3_file*,
                int flags, int *pOutFlags);
   int (*xDelete)(sqlite3_vfs*, const char *zName, int syncDir);
   int (*xAccess)(sqlite3_vfs*, const char *zName, int flags, int *pResOut);
@@ -1390,7 +1508,7 @@ struct sqlite3_vfs {
   /*
   ** The methods above are in versions 1 through 3 of the sqlite_vfs object.
   ** New fields may be appended in future versions.  The iVersion
-  ** value will increment whenever this happens. 
+  ** value will increment whenever this happens.
   */
 };
 
@@ -1434,7 +1552,7 @@ struct sqlite3_vfs {
 ** </ul>
 **
 ** When unlocking, the same SHARED or EXCLUSIVE flag must be supplied as
-** was given on the corresponding lock.  
+** was given on the corresponding lock.
 **
 ** The xShmLock method can transition between unlocked and SHARED or
 ** between unlocked and EXCLUSIVE.  It cannot transition between SHARED
@@ -1549,20 +1667,23 @@ SQLITE_API int sqlite3_os_end(void);
 ** must ensure that no other SQLite interfaces are invoked by other
 ** threads while sqlite3_config() is running.</b>
 **
-** The sqlite3_config() interface
-** may only be invoked prior to library initialization using
-** [sqlite3_initialize()] or after shutdown by [sqlite3_shutdown()].
-** ^If sqlite3_config() is called after [sqlite3_initialize()] and before
-** [sqlite3_shutdown()] then it will return SQLITE_MISUSE.
-** Note, however, that ^sqlite3_config() can be called as part of the
-** implementation of an application-defined [sqlite3_os_init()].
-**
 ** The first argument to sqlite3_config() is an integer
 ** [configuration option] that determines
 ** what property of SQLite is to be configured.  Subsequent arguments
 ** vary depending on the [configuration option]
 ** in the first argument.
 **
+** For most configuration options, the sqlite3_config() interface
+** may only be invoked prior to library initialization using
+** [sqlite3_initialize()] or after shutdown by [sqlite3_shutdown()].
+** The exceptional configuration options that may be invoked at any time
+** are called "anytime configuration options".
+** ^If sqlite3_config() is called after [sqlite3_initialize()] and before
+** [sqlite3_shutdown()] with a first argument that is not an anytime
+** configuration option, then the sqlite3_config() call will return SQLITE_MISUSE.
+** Note, however, that ^sqlite3_config() can be called as part of the
+** implementation of an application-defined [sqlite3_os_init()].
+**
 ** ^When a configuration option is set, sqlite3_config() returns [SQLITE_OK].
 ** ^If the option is unknown or SQLite is unable to set the option
 ** then this routine returns a non-zero [error code].
@@ -1579,7 +1700,7 @@ SQLITE_API int sqlite3_config(int, ...);
 ** [database connection] (specified in the first argument).
 **
 ** The second argument to sqlite3_db_config(D,V,...)  is the
-** [SQLITE_DBCONFIG_LOOKASIDE | configuration verb] - an integer code 
+** [SQLITE_DBCONFIG_LOOKASIDE | configuration verb] - an integer code
 ** that indicates what aspect of the [database connection] is being configured.
 ** Subsequent arguments vary depending on the configuration verb.
 **
@@ -1597,7 +1718,7 @@ SQLITE_API int sqlite3_db_config(sqlite3*, int op, ...);
 ** This object is used in only one place in the SQLite interface.
 ** A pointer to an instance of this object is the argument to
 ** [sqlite3_config()] when the configuration option is
-** [SQLITE_CONFIG_MALLOC] or [SQLITE_CONFIG_GETMALLOC].  
+** [SQLITE_CONFIG_MALLOC] or [SQLITE_CONFIG_GETMALLOC].
 ** By creating an instance of this object
 ** and passing it to [sqlite3_config]([SQLITE_CONFIG_MALLOC])
 ** during configuration, an application can specify an alternative
@@ -1627,7 +1748,7 @@ SQLITE_API int sqlite3_db_config(sqlite3*, int op, ...);
 ** allocators round up memory allocations at least to the next multiple
 ** of 8.  Some allocators round up to a larger multiple or to a power of 2.
 ** Every memory allocation request coming in through [sqlite3_malloc()]
-** or [sqlite3_realloc()] first calls xRoundup.  If xRoundup returns 0, 
+** or [sqlite3_realloc()] first calls xRoundup.  If xRoundup returns 0,
 ** that causes the corresponding memory allocation to fail.
 **
 ** The xInit method initializes the memory allocator.  For example,
@@ -1637,7 +1758,7 @@ SQLITE_API int sqlite3_db_config(sqlite3*, int op, ...);
 ** by xInit.  The pAppData pointer is used as the only parameter to
 ** xInit and xShutdown.
 **
-** SQLite holds the [SQLITE_MUTEX_STATIC_MASTER] mutex when it invokes
+** SQLite holds the [SQLITE_MUTEX_STATIC_MAIN] mutex when it invokes
 ** the xInit method, so the xInit method need not be threadsafe.  The
 ** xShutdown method is only called from [sqlite3_shutdown()] so it does
 ** not need to be threadsafe either.  For all other methods, SQLite
@@ -1670,6 +1791,23 @@ struct sqlite3_mem_methods {
 ** These constants are the available integer configuration options that
 ** can be passed as the first argument to the [sqlite3_config()] interface.
 **
+** Most of the configuration options for sqlite3_config()
+** will only work if invoked prior to [sqlite3_initialize()] or after
+** [sqlite3_shutdown()].  The few exceptions to this rule are called
+** "anytime configuration options".
+** ^Calling [sqlite3_config()] with a first argument that is not an
+** anytime configuration option in between calls to [sqlite3_initialize()] and
+** [sqlite3_shutdown()] is a no-op that returns SQLITE_MISUSE.
+**
+** The set of anytime configuration options can change (by insertions
+** and/or deletions) from one release of SQLite to the next.
+** As of SQLite version 3.42.0, the complete set of anytime configuration
+** options is:
+** <ul>
+** <li> SQLITE_CONFIG_LOG
+** <li> SQLITE_CONFIG_PCACHE_HDRSZ
+** </ul>
+**
 ** New configuration options may be added in future releases of SQLite.
 ** Existing configuration options might be discontinued.  Applications
 ** should check the return code from [sqlite3_config()] to make sure that
@@ -1685,7 +1823,7 @@ struct sqlite3_mem_methods {
 ** by a single thread.   ^If SQLite is compiled with
 ** the [SQLITE_THREADSAFE | SQLITE_THREADSAFE=0] compile-time option then
 ** it is not possible to change the [threading mode] from its default
-** value of Single-thread and so [sqlite3_config()] will return 
+** value of Single-thread and so [sqlite3_config()] will return
 ** [SQLITE_ERROR] if called with the SQLITE_CONFIG_SINGLETHREAD
 ** configuration option.</dd>
 **
@@ -1720,7 +1858,7 @@ struct sqlite3_mem_methods {
 ** SQLITE_CONFIG_SERIALIZED configuration option.</dd>
 **
 ** [[SQLITE_CONFIG_MALLOC]] <dt>SQLITE_CONFIG_MALLOC</dt>
-** <dd> ^(The SQLITE_CONFIG_MALLOC option takes a single argument which is 
+** <dd> ^(The SQLITE_CONFIG_MALLOC option takes a single argument which is
 ** a pointer to an instance of the [sqlite3_mem_methods] structure.
 ** The argument specifies
 ** alternative low-level memory allocation routines to be used in place of
@@ -1771,7 +1909,7 @@ struct sqlite3_mem_methods {
 ** [[SQLITE_CONFIG_PAGECACHE]] <dt>SQLITE_CONFIG_PAGECACHE</dt>
 ** <dd> ^The SQLITE_CONFIG_PAGECACHE option specifies a memory pool
 ** that SQLite can use for the database page cache with the default page
-** cache implementation.  
+** cache implementation.
 ** This configuration option is a no-op if an application-defined page
 ** cache implementation is loaded using the [SQLITE_CONFIG_PCACHE2].
 ** ^There are three arguments to SQLITE_CONFIG_PAGECACHE: A pointer to
@@ -1799,7 +1937,7 @@ struct sqlite3_mem_methods {
 ** additional cache line. </dd>
 **
 ** [[SQLITE_CONFIG_HEAP]] <dt>SQLITE_CONFIG_HEAP</dt>
-** <dd> ^The SQLITE_CONFIG_HEAP option specifies a static memory buffer 
+** <dd> ^The SQLITE_CONFIG_HEAP option specifies a static memory buffer
 ** that SQLite will use for all of its dynamic memory allocation needs
 ** beyond those provided for by [SQLITE_CONFIG_PAGECACHE].
 ** ^The SQLITE_CONFIG_HEAP option is only available if SQLite is compiled
@@ -1854,7 +1992,7 @@ struct sqlite3_mem_methods {
 ** configuration on individual connections.)^ </dd>
 **
 ** [[SQLITE_CONFIG_PCACHE2]] <dt>SQLITE_CONFIG_PCACHE2</dt>
-** <dd> ^(The SQLITE_CONFIG_PCACHE2 option takes a single argument which is 
+** <dd> ^(The SQLITE_CONFIG_PCACHE2 option takes a single argument which is
 ** a pointer to an [sqlite3_pcache_methods2] object.  This object specifies
 ** the interface to a custom page cache implementation.)^
 ** ^SQLite makes a copy of the [sqlite3_pcache_methods2] object.</dd>
@@ -1868,7 +2006,7 @@ struct sqlite3_mem_methods {
 ** <dd> The SQLITE_CONFIG_LOG option is used to configure the SQLite
 ** global [error log].
 ** (^The SQLITE_CONFIG_LOG option takes two arguments: a pointer to a
-** function with a call signature of void(*)(void*,int,const char*), 
+** function with a call signature of void(*)(void*,int,const char*),
 ** and a pointer to void. ^If the function pointer is not NULL, it is
 ** invoked by [sqlite3_log()] to process each logging event.  ^If the
 ** function pointer is NULL, the [sqlite3_log()] interface becomes a no-op.
@@ -1977,7 +2115,7 @@ struct sqlite3_mem_methods {
 ** [[SQLITE_CONFIG_STMTJRNL_SPILL]]
 ** <dt>SQLITE_CONFIG_STMTJRNL_SPILL
 ** <dd>^The SQLITE_CONFIG_STMTJRNL_SPILL option takes a single parameter which
-** becomes the [statement journal] spill-to-disk threshold.  
+** becomes the [statement journal] spill-to-disk threshold.
 ** [Statement journals] are held in memory until their size (in bytes)
 ** exceeds this threshold, at which point they are written to disk.
 ** Or if the threshold is -1, statement journals are always held
@@ -1999,8 +2137,8 @@ struct sqlite3_mem_methods {
 ** than the configured sorter-reference size threshold - then a reference
 ** is stored in each sorted record and the required column values loaded
 ** from the database as records are returned in sorted order. The default
-** value for this option is to never use this optimization. Specifying a 
-** negative value for this option restores the default behaviour.
+** value for this option is to never use this optimization. Specifying a
+** negative value for this option restores the default behavior.
 ** This option is only available if SQLite is compiled with the
 ** [SQLITE_ENABLE_SORTER_REFERENCES] compile-time option.
 **
@@ -2014,30 +2152,46 @@ struct sqlite3_mem_methods {
 ** configuration setting is never used, then the default maximum is determined
 ** by the [SQLITE_MEMDB_DEFAULT_MAXSIZE] compile-time option.  If that
 ** compile-time option is not set, then the default maximum is 1073741824.
+**
+** [[SQLITE_CONFIG_ROWID_IN_VIEW]]
+** <dt>SQLITE_CONFIG_ROWID_IN_VIEW
+** <dd>The SQLITE_CONFIG_ROWID_IN_VIEW option enables or disables the ability
+** for VIEWs to have a ROWID.  The capability can only be enabled if SQLite is
+** compiled with -DSQLITE_ALLOW_ROWID_IN_VIEW, in which case the capability
+** defaults to on.  This configuration option queries the current setting or
+** changes the setting to off or on.  The argument is a pointer to an integer.
+** If that integer initially holds a value of 1, then the ability for VIEWs to
+** have ROWIDs is activated.  If the integer initially holds zero, then the
+** ability is deactivated.  Any other initial value for the integer leaves the
+** setting unchanged.  After changes, if any, the integer is written with
+** a 1 or 0, if the ability for VIEWs to have ROWIDs is on or off.  If SQLite
+** is compiled without -DSQLITE_ALLOW_ROWID_IN_VIEW (which is the usual and
+** recommended case) then the integer is always filled with zero, regardless
+** if its initial value.
 ** </dl>
 */
-#define SQLITE_CONFIG_SINGLETHREAD  1  /* nil */
-#define SQLITE_CONFIG_MULTITHREAD   2  /* nil */
-#define SQLITE_CONFIG_SERIALIZED    3  /* nil */
-#define SQLITE_CONFIG_MALLOC        4  /* sqlite3_mem_methods* */
-#define SQLITE_CONFIG_GETMALLOC     5  /* sqlite3_mem_methods* */
-#define SQLITE_CONFIG_SCRATCH       6  /* No longer used */
-#define SQLITE_CONFIG_PAGECACHE     7  /* void*, int sz, int N */
-#define SQLITE_CONFIG_HEAP          8  /* void*, int nByte, int min */
-#define SQLITE_CONFIG_MEMSTATUS     9  /* boolean */
-#define SQLITE_CONFIG_MUTEX        10  /* sqlite3_mutex_methods* */
-#define SQLITE_CONFIG_GETMUTEX     11  /* sqlite3_mutex_methods* */
-/* previously SQLITE_CONFIG_CHUNKALLOC 12 which is now unused. */ 
-#define SQLITE_CONFIG_LOOKASIDE    13  /* int int */
-#define SQLITE_CONFIG_PCACHE       14  /* no-op */
-#define SQLITE_CONFIG_GETPCACHE    15  /* no-op */
-#define SQLITE_CONFIG_LOG          16  /* xFunc, void* */
-#define SQLITE_CONFIG_URI          17  /* int */
-#define SQLITE_CONFIG_PCACHE2      18  /* sqlite3_pcache_methods2* */
-#define SQLITE_CONFIG_GETPCACHE2   19  /* sqlite3_pcache_methods2* */
+#define SQLITE_CONFIG_SINGLETHREAD         1  /* nil */
+#define SQLITE_CONFIG_MULTITHREAD          2  /* nil */
+#define SQLITE_CONFIG_SERIALIZED           3  /* nil */
+#define SQLITE_CONFIG_MALLOC               4  /* sqlite3_mem_methods* */
+#define SQLITE_CONFIG_GETMALLOC            5  /* sqlite3_mem_methods* */
+#define SQLITE_CONFIG_SCRATCH              6  /* No longer used */
+#define SQLITE_CONFIG_PAGECACHE            7  /* void*, int sz, int N */
+#define SQLITE_CONFIG_HEAP                 8  /* void*, int nByte, int min */
+#define SQLITE_CONFIG_MEMSTATUS            9  /* boolean */
+#define SQLITE_CONFIG_MUTEX               10  /* sqlite3_mutex_methods* */
+#define SQLITE_CONFIG_GETMUTEX            11  /* sqlite3_mutex_methods* */
+/* previously SQLITE_CONFIG_CHUNKALLOC    12 which is now unused. */
+#define SQLITE_CONFIG_LOOKASIDE           13  /* int int */
+#define SQLITE_CONFIG_PCACHE              14  /* no-op */
+#define SQLITE_CONFIG_GETPCACHE           15  /* no-op */
+#define SQLITE_CONFIG_LOG                 16  /* xFunc, void* */
+#define SQLITE_CONFIG_URI                 17  /* int */
+#define SQLITE_CONFIG_PCACHE2             18  /* sqlite3_pcache_methods2* */
+#define SQLITE_CONFIG_GETPCACHE2          19  /* sqlite3_pcache_methods2* */
 #define SQLITE_CONFIG_COVERING_INDEX_SCAN 20  /* int */
-#define SQLITE_CONFIG_SQLLOG       21  /* xSqllog, void* */
-#define SQLITE_CONFIG_MMAP_SIZE    22  /* sqlite3_int64, sqlite3_int64 */
+#define SQLITE_CONFIG_SQLLOG              21  /* xSqllog, void* */
+#define SQLITE_CONFIG_MMAP_SIZE           22  /* sqlite3_int64, sqlite3_int64 */
 #define SQLITE_CONFIG_WIN32_HEAPSIZE      23  /* int nByte */
 #define SQLITE_CONFIG_PCACHE_HDRSZ        24  /* int *psz */
 #define SQLITE_CONFIG_PMASZ               25  /* unsigned int szPma */
@@ -2045,6 +2199,7 @@ struct sqlite3_mem_methods {
 #define SQLITE_CONFIG_SMALL_MALLOC        27  /* boolean */
 #define SQLITE_CONFIG_SORTERREF_SIZE      28  /* int nByte */
 #define SQLITE_CONFIG_MEMDB_MAXSIZE       29  /* sqlite3_int64 */
+#define SQLITE_CONFIG_ROWID_IN_VIEW       30  /* int* */
 
 /*
 ** CAPI3REF: Database Connection Configuration Options
@@ -2062,7 +2217,7 @@ struct sqlite3_mem_methods {
 ** <dl>
 ** [[SQLITE_DBCONFIG_LOOKASIDE]]
 ** <dt>SQLITE_DBCONFIG_LOOKASIDE</dt>
-** <dd> ^This option takes three additional arguments that determine the 
+** <dd> ^This option takes three additional arguments that determine the
 ** [lookaside memory allocator] configuration for the [database connection].
 ** ^The first argument (the third parameter to [sqlite3_db_config()] is a
 ** pointer to a memory buffer to use for lookaside memory.
@@ -2078,9 +2233,9 @@ struct sqlite3_mem_methods {
 ** configuration for a database connection can only be changed when that
 ** connection is not currently using lookaside memory, or in other words
 ** when the "current value" returned by
-** [sqlite3_db_status](D,[SQLITE_CONFIG_LOOKASIDE],...) is zero.
+** [sqlite3_db_status](D,[SQLITE_DBSTATUS_LOOKASIDE_USED],...) is zero.
 ** Any attempt to change the lookaside memory configuration when lookaside
-** memory is in use leaves the configuration unchanged and returns 
+** memory is in use leaves the configuration unchanged and returns
 ** [SQLITE_BUSY].)^</dd>
 **
 ** [[SQLITE_DBCONFIG_ENABLE_FKEY]]
@@ -2103,7 +2258,13 @@ struct sqlite3_mem_methods {
 ** The second parameter is a pointer to an integer into which
 ** is written 0 or 1 to indicate whether triggers are disabled or enabled
 ** following this call.  The second parameter may be a NULL pointer, in
-** which case the trigger setting is not reported back. </dd>
+** which case the trigger setting is not reported back.
+**
+** <p>Originally this option disabled all triggers.  ^(However, since
+** SQLite version 3.35.0, TEMP triggers are still allowed even if
+** this option is off.  So, in other words, this option now only disables
+** triggers in the main database schema or in the schemas of ATTACH-ed
+** databases.)^ </dd>
 **
 ** [[SQLITE_DBCONFIG_ENABLE_VIEW]]
 ** <dt>SQLITE_DBCONFIG_ENABLE_VIEW</dt>
@@ -2114,7 +2275,13 @@ struct sqlite3_mem_methods {
 ** The second parameter is a pointer to an integer into which
 ** is written 0 or 1 to indicate whether views are disabled or enabled
 ** following this call.  The second parameter may be a NULL pointer, in
-** which case the view setting is not reported back. </dd>
+** which case the view setting is not reported back.
+**
+** <p>Originally this option disabled all views.  ^(However, since
+** SQLite version 3.35.0, TEMP views are still allowed even if
+** this option is off.  So, in other words, this option now only disables
+** views in the main database schema or in the schemas of ATTACH-ed
+** databases.)^ </dd>
 **
 ** [[SQLITE_DBCONFIG_ENABLE_FTS3_TOKENIZER]]
 ** <dt>SQLITE_DBCONFIG_ENABLE_FTS3_TOKENIZER</dt>
@@ -2157,13 +2324,13 @@ struct sqlite3_mem_methods {
 ** until after the database connection closes.
 ** </dd>
 **
-** [[SQLITE_DBCONFIG_NO_CKPT_ON_CLOSE]] 
+** [[SQLITE_DBCONFIG_NO_CKPT_ON_CLOSE]]
 ** <dt>SQLITE_DBCONFIG_NO_CKPT_ON_CLOSE</dt>
-** <dd> Usually, when a database in wal mode is closed or detached from a 
-** database handle, SQLite checks if this will mean that there are now no 
-** connections at all to the database. If so, it performs a checkpoint 
+** <dd> Usually, when a database in wal mode is closed or detached from a
+** database handle, SQLite checks if this will mean that there are now no
+** connections at all to the database. If so, it performs a checkpoint
 ** operation before closing the connection. This option may be used to
-** override this behaviour. The first parameter passed to this operation
+** override this behavior. The first parameter passed to this operation
 ** is an integer - positive to disable checkpoints-on-close, or zero (the
 ** default) to enable them, and negative to leave the setting unchanged.
 ** The second parameter is a pointer to an integer
@@ -2180,7 +2347,7 @@ struct sqlite3_mem_methods {
 ** slower.  But the QPSG has the advantage of more predictable behavior.  With
 ** the QPSG active, SQLite will always use the same query plan in the field as
 ** was used during testing in the lab.
-** The first argument to this setting is an integer which is 0 to disable 
+** The first argument to this setting is an integer which is 0 to disable
 ** the QPSG, positive to enable QPSG, or negative to leave the setting
 ** unchanged. The second parameter is a pointer to an integer into which
 ** is written 0 or 1 to indicate whether the QPSG is disabled or enabled
@@ -2188,15 +2355,15 @@ struct sqlite3_mem_methods {
 ** </dd>
 **
 ** [[SQLITE_DBCONFIG_TRIGGER_EQP]] <dt>SQLITE_DBCONFIG_TRIGGER_EQP</dt>
-** <dd> By default, the output of EXPLAIN QUERY PLAN commands does not 
+** <dd> By default, the output of EXPLAIN QUERY PLAN commands does not
 ** include output for any operations performed by trigger programs. This
 ** option is used to set or clear (the default) a flag that governs this
 ** behavior. The first parameter passed to this operation is an integer -
 ** positive to enable output for trigger programs, or zero to disable it,
 ** or negative to leave the setting unchanged.
-** The second parameter is a pointer to an integer into which is written 
-** 0 or 1 to indicate whether output-for-triggers has been disabled - 0 if 
-** it is not disabled, 1 if it is.  
+** The second parameter is a pointer to an integer into which is written
+** 0 or 1 to indicate whether output-for-triggers has been disabled - 0 if
+** it is not disabled, 1 if it is.
 ** </dd>
 **
 ** [[SQLITE_DBCONFIG_RESET_DATABASE]] <dt>SQLITE_DBCONFIG_RESET_DATABASE</dt>
@@ -2210,24 +2377,29 @@ struct sqlite3_mem_methods {
 **      database, or calling sqlite3_table_column_metadata(), ignoring any
 **      errors.  This step is only necessary if the application desires to keep
 **      the database in WAL mode after the reset if it was in WAL mode before
-**      the reset.  
+**      the reset.
 ** <li> sqlite3_db_config(db, SQLITE_DBCONFIG_RESET_DATABASE, 1, 0);
 ** <li> [sqlite3_exec](db, "[VACUUM]", 0, 0, 0);
 ** <li> sqlite3_db_config(db, SQLITE_DBCONFIG_RESET_DATABASE, 0, 0);
 ** </ol>
 ** Because resetting a database is destructive and irreversible, the
-** process requires the use of this obscure API and multiple steps to help
-** ensure that it does not happen by accident.
+** process requires the use of this obscure API and multiple steps to
+** help ensure that it does not happen by accident. Because this
+** feature must be capable of resetting corrupt databases, and
+** shutting down virtual tables may require access to that corrupt
+** storage, the library must abandon any installed virtual tables
+** without calling their xDestroy() methods.
 **
 ** [[SQLITE_DBCONFIG_DEFENSIVE]] <dt>SQLITE_DBCONFIG_DEFENSIVE</dt>
 ** <dd>The SQLITE_DBCONFIG_DEFENSIVE option activates or deactivates the
 ** "defensive" flag for a database connection.  When the defensive
-** flag is enabled, language features that allow ordinary SQL to 
+** flag is enabled, language features that allow ordinary SQL to
 ** deliberately corrupt the database file are disabled.  The disabled
 ** features include but are not limited to the following:
 ** <ul>
 ** <li> The [PRAGMA writable_schema=ON] statement.
 ** <li> The [PRAGMA journal_mode=OFF] statement.
+** <li> The [PRAGMA schema_version=N] statement.
 ** <li> Writes to the [sqlite_dbpage] virtual table.
 ** <li> Direct writes to [shadow tables].
 ** </ul>
@@ -2237,7 +2409,7 @@ struct sqlite3_mem_methods {
 ** <dd>The SQLITE_DBCONFIG_WRITABLE_SCHEMA option activates or deactivates the
 ** "writable_schema" flag. This has the same effect and is logically equivalent
 ** to setting [PRAGMA writable_schema=ON] or [PRAGMA writable_schema=OFF].
-** The first argument to this setting is an integer which is 0 to disable 
+** The first argument to this setting is an integer which is 0 to disable
 ** the writable_schema, positive to enable writable_schema, or negative to
 ** leave the setting unchanged. The second parameter is a pointer to an
 ** integer into which is written 0 or 1 to indicate whether the writable_schema
@@ -2255,7 +2427,7 @@ struct sqlite3_mem_methods {
 ** </dd>
 **
 ** [[SQLITE_DBCONFIG_DQS_DML]]
-** <dt>SQLITE_DBCONFIG_DQS_DML</td>
+** <dt>SQLITE_DBCONFIG_DQS_DML</dt>
 ** <dd>The SQLITE_DBCONFIG_DQS_DML option activates or deactivates
 ** the legacy [double-quoted string literal] misfeature for DML statements
 ** only, that is DELETE, INSERT, SELECT, and UPDATE statements. The
@@ -2264,7 +2436,7 @@ struct sqlite3_mem_methods {
 ** </dd>
 **
 ** [[SQLITE_DBCONFIG_DQS_DDL]]
-** <dt>SQLITE_DBCONFIG_DQS_DDL</td>
+** <dt>SQLITE_DBCONFIG_DQS_DDL</dt>
 ** <dd>The SQLITE_DBCONFIG_DQS option activates or deactivates
 ** the legacy [double-quoted string literal] misfeature for DDL statements,
 ** such as CREATE TABLE and CREATE INDEX. The
@@ -2273,16 +2445,15 @@ struct sqlite3_mem_methods {
 ** </dd>
 **
 ** [[SQLITE_DBCONFIG_TRUSTED_SCHEMA]]
-** <dt>SQLITE_DBCONFIG_TRUSTED_SCHEMA</td>
+** <dt>SQLITE_DBCONFIG_TRUSTED_SCHEMA</dt>
 ** <dd>The SQLITE_DBCONFIG_TRUSTED_SCHEMA option tells SQLite to
-** assume that database schemas (the contents of the [sqlite_master] tables)
-** are untainted by malicious content.
+** assume that database schemas are untainted by malicious content.
 ** When the SQLITE_DBCONFIG_TRUSTED_SCHEMA option is disabled, SQLite
 ** takes additional defensive steps to protect the application from harm
 ** including:
 ** <ul>
 ** <li> Prohibit the use of SQL functions inside triggers, views,
-** CHECK constraints, DEFAULT clauses, expression indexes, 
+** CHECK constraints, DEFAULT clauses, expression indexes,
 ** partial indexes, or generated columns
 ** unless those functions are tagged with [SQLITE_INNOCUOUS].
 ** <li> Prohibit the use of virtual tables inside of triggers or views
@@ -2294,7 +2465,7 @@ struct sqlite3_mem_methods {
 ** </dd>
 **
 ** [[SQLITE_DBCONFIG_LEGACY_FILE_FORMAT]]
-** <dt>SQLITE_DBCONFIG_LEGACY_FILE_FORMAT</td>
+** <dt>SQLITE_DBCONFIG_LEGACY_FILE_FORMAT</dt>
 ** <dd>The SQLITE_DBCONFIG_LEGACY_FILE_FORMAT option activates or deactivates
 ** the legacy file format flag.  When activated, this flag causes all newly
 ** created database file to have a schema format version number (the 4-byte
@@ -2303,7 +2474,7 @@ struct sqlite3_mem_methods {
 ** any SQLite version back to 3.0.0 ([dateof:3.0.0]).  Without this setting,
 ** newly created databases are generally not understandable by SQLite versions
 ** prior to 3.3.0 ([dateof:3.3.0]).  As these words are written, there
-** is now scarcely any need to generated database files that are compatible 
+** is now scarcely any need to generate database files that are compatible
 ** all the way back to version 3.0.0, and so this setting is of little
 ** practical use, but is provided so that SQLite can continue to claim the
 ** ability to generate new database files that are compatible with  version
@@ -2312,8 +2483,40 @@ struct sqlite3_mem_methods {
 ** the [VACUUM] command will fail with an obscure error when attempting to
 ** process a table with generated columns and a descending index.  This is
 ** not considered a bug since SQLite versions 3.3.0 and earlier do not support
-** either generated columns or decending indexes.
+** either generated columns or descending indexes.
+** </dd>
+**
+** [[SQLITE_DBCONFIG_STMT_SCANSTATUS]]
+** <dt>SQLITE_DBCONFIG_STMT_SCANSTATUS</dt>
+** <dd>The SQLITE_DBCONFIG_STMT_SCANSTATUS option is only useful in
+** SQLITE_ENABLE_STMT_SCANSTATUS builds. In this case, it sets or clears
+** a flag that enables collection of the sqlite3_stmt_scanstatus_v2()
+** statistics. For statistics to be collected, the flag must be set on
+** the database handle both when the SQL statement is prepared and when it
+** is stepped. The flag is set (collection of statistics is enabled)
+** by default.  This option takes two arguments: an integer and a pointer to
+** an integer..  The first argument is 1, 0, or -1 to enable, disable, or
+** leave unchanged the statement scanstatus option.  If the second argument
+** is not NULL, then the value of the statement scanstatus setting after
+** processing the first argument is written into the integer that the second
+** argument points to.
+** </dd>
+**
+** [[SQLITE_DBCONFIG_REVERSE_SCANORDER]]
+** <dt>SQLITE_DBCONFIG_REVERSE_SCANORDER</dt>
+** <dd>The SQLITE_DBCONFIG_REVERSE_SCANORDER option changes the default order
+** in which tables and indexes are scanned so that the scans start at the end
+** and work toward the beginning rather than starting at the beginning and
+** working toward the end. Setting SQLITE_DBCONFIG_REVERSE_SCANORDER is the
+** same as setting [PRAGMA reverse_unordered_selects].  This option takes
+** two arguments which are an integer and a pointer to an integer.  The first
+** argument is 1, 0, or -1 to enable, disable, or leave unchanged the
+** reverse scan order flag, respectively.  If the second argument is not NULL,
+** then 0 or 1 is written into the integer that the second argument points to
+** depending on if the reverse scan order flag is set after processing the
+** first argument.
 ** </dd>
+**
 ** </dl>
 */
 #define SQLITE_DBCONFIG_MAINDBNAME            1000 /* const char* */
@@ -2334,7 +2537,9 @@ struct sqlite3_mem_methods {
 #define SQLITE_DBCONFIG_ENABLE_VIEW           1015 /* int int* */
 #define SQLITE_DBCONFIG_LEGACY_FILE_FORMAT    1016 /* int int* */
 #define SQLITE_DBCONFIG_TRUSTED_SCHEMA        1017 /* int int* */
-#define SQLITE_DBCONFIG_MAX                   1017 /* Largest DBCONFIG */
+#define SQLITE_DBCONFIG_STMT_SCANSTATUS       1018 /* int int* */
+#define SQLITE_DBCONFIG_REVERSE_SCANORDER     1019 /* int int* */
+#define SQLITE_DBCONFIG_MAX                   1019 /* Largest DBCONFIG */
 
 /*
 ** CAPI3REF: Enable Or Disable Extended Result Codes
@@ -2361,8 +2566,8 @@ SQLITE_API int sqlite3_extended_result_codes(sqlite3*, int onoff);
 ** ^The sqlite3_last_insert_rowid(D) interface usually returns the [rowid] of
 ** the most recent successful [INSERT] into a rowid table or [virtual table]
 ** on database connection D. ^Inserts into [WITHOUT ROWID] tables are not
-** recorded. ^If no successful [INSERT]s into rowid tables have ever occurred 
-** on the database connection D, then sqlite3_last_insert_rowid(D) returns 
+** recorded. ^If no successful [INSERT]s into rowid tables have ever occurred
+** on the database connection D, then sqlite3_last_insert_rowid(D) returns
 ** zero.
 **
 ** As well as being set automatically as rows are inserted into database
@@ -2372,15 +2577,15 @@ SQLITE_API int sqlite3_extended_result_codes(sqlite3*, int onoff);
 ** Some virtual table implementations may INSERT rows into rowid tables as
 ** part of committing a transaction (e.g. to flush data accumulated in memory
 ** to disk). In this case subsequent calls to this function return the rowid
-** associated with these internal INSERT operations, which leads to 
+** associated with these internal INSERT operations, which leads to
 ** unintuitive results. Virtual table implementations that do write to rowid
-** tables in this way can avoid this problem by restoring the original 
-** rowid value using [sqlite3_set_last_insert_rowid()] before returning 
+** tables in this way can avoid this problem by restoring the original
+** rowid value using [sqlite3_set_last_insert_rowid()] before returning
 ** control to the user.
 **
-** ^(If an [INSERT] occurs within a trigger then this routine will 
-** return the [rowid] of the inserted row as long as the trigger is 
-** running. Once the trigger program ends, the value returned 
+** ^(If an [INSERT] occurs within a trigger then this routine will
+** return the [rowid] of the inserted row as long as the trigger is
+** running. Once the trigger program ends, the value returned
 ** by this routine reverts to what it was before the trigger was fired.)^
 **
 ** ^An [INSERT] that fails due to a constraint violation is not a
@@ -2413,7 +2618,7 @@ SQLITE_API sqlite3_int64 sqlite3_last_insert_rowid(sqlite3*);
 ** METHOD: sqlite3
 **
 ** The sqlite3_set_last_insert_rowid(D, R) method allows the application to
-** set the value returned by calling sqlite3_last_insert_rowid(D) to R 
+** set the value returned by calling sqlite3_last_insert_rowid(D) to R
 ** without inserting a row into the database.
 */
 SQLITE_API void sqlite3_set_last_insert_rowid(sqlite3*,sqlite3_int64);
@@ -2422,44 +2627,47 @@ SQLITE_API void sqlite3_set_last_insert_rowid(sqlite3*,sqlite3_int64);
 ** CAPI3REF: Count The Number Of Rows Modified
 ** METHOD: sqlite3
 **
-** ^This function returns the number of rows modified, inserted or
+** ^These functions return the number of rows modified, inserted or
 ** deleted by the most recently completed INSERT, UPDATE or DELETE
 ** statement on the database connection specified by the only parameter.
-** ^Executing any other type of SQL statement does not modify the value
-** returned by this function.
+** The two functions are identical except for the type of the return value
+** and that if the number of rows modified by the most recent INSERT, UPDATE
+** or DELETE is greater than the maximum value supported by type "int", then
+** the return value of sqlite3_changes() is undefined. ^Executing any other
+** type of SQL statement does not modify the value returned by these functions.
 **
 ** ^Only changes made directly by the INSERT, UPDATE or DELETE statement are
-** considered - auxiliary changes caused by [CREATE TRIGGER | triggers], 
+** considered - auxiliary changes caused by [CREATE TRIGGER | triggers],
 ** [foreign key actions] or [REPLACE] constraint resolution are not counted.
-** 
-** Changes to a view that are intercepted by 
-** [INSTEAD OF trigger | INSTEAD OF triggers] are not counted. ^The value 
-** returned by sqlite3_changes() immediately after an INSERT, UPDATE or 
-** DELETE statement run on a view is always zero. Only changes made to real 
+**
+** Changes to a view that are intercepted by
+** [INSTEAD OF trigger | INSTEAD OF triggers] are not counted. ^The value
+** returned by sqlite3_changes() immediately after an INSERT, UPDATE or
+** DELETE statement run on a view is always zero. Only changes made to real
 ** tables are counted.
 **
 ** Things are more complicated if the sqlite3_changes() function is
 ** executed while a trigger program is running. This may happen if the
 ** program uses the [changes() SQL function], or if some other callback
 ** function invokes sqlite3_changes() directly. Essentially:
-** 
+**
 ** <ul>
 **   <li> ^(Before entering a trigger program the value returned by
-**        sqlite3_changes() function is saved. After the trigger program 
+**        sqlite3_changes() function is saved. After the trigger program
 **        has finished, the original value is restored.)^
-** 
-**   <li> ^(Within a trigger program each INSERT, UPDATE and DELETE 
-**        statement sets the value returned by sqlite3_changes() 
-**        upon completion as normal. Of course, this value will not include 
-**        any changes performed by sub-triggers, as the sqlite3_changes() 
+**
+**   <li> ^(Within a trigger program each INSERT, UPDATE and DELETE
+**        statement sets the value returned by sqlite3_changes()
+**        upon completion as normal. Of course, this value will not include
+**        any changes performed by sub-triggers, as the sqlite3_changes()
 **        value will be saved and restored after each sub-trigger has run.)^
 ** </ul>
-** 
+**
 ** ^This means that if the changes() SQL function (or similar) is used
-** by the first INSERT, UPDATE or DELETE statement within a trigger, it 
+** by the first INSERT, UPDATE or DELETE statement within a trigger, it
 ** returns the value as set when the calling statement began executing.
-** ^If it is used by the second or subsequent such statement within a trigger 
-** program, the value returned reflects the number of rows modified by the 
+** ^If it is used by the second or subsequent such statement within a trigger
+** program, the value returned reflects the number of rows modified by the
 ** previous INSERT, UPDATE or DELETE statement within the same trigger.
 **
 ** If a separate thread makes changes on the same database connection
@@ -2475,20 +2683,25 @@ SQLITE_API void sqlite3_set_last_insert_rowid(sqlite3*,sqlite3_int64);
 ** </ul>
 */
 SQLITE_API int sqlite3_changes(sqlite3*);
+SQLITE_API sqlite3_int64 sqlite3_changes64(sqlite3*);
 
 /*
 ** CAPI3REF: Total Number Of Rows Modified
 ** METHOD: sqlite3
 **
-** ^This function returns the total number of rows inserted, modified or
+** ^These functions return the total number of rows inserted, modified or
 ** deleted by all [INSERT], [UPDATE] or [DELETE] statements completed
 ** since the database connection was opened, including those executed as
-** part of trigger programs. ^Executing any other type of SQL statement
-** does not affect the value returned by sqlite3_total_changes().
-** 
+** part of trigger programs. The two functions are identical except for the
+** type of the return value and that if the number of rows modified by the
+** connection exceeds the maximum value supported by type "int", then
+** the return value of sqlite3_total_changes() is undefined. ^Executing
+** any other type of SQL statement does not affect the value returned by
+** sqlite3_total_changes().
+**
 ** ^Changes made as part of [foreign key actions] are included in the
 ** count, but those made as part of REPLACE constraint resolution are
-** not. ^Changes to a view that are intercepted by INSTEAD OF triggers 
+** not. ^Changes to a view that are intercepted by INSTEAD OF triggers
 ** are not counted.
 **
 ** The [sqlite3_total_changes(D)] interface only reports the number
@@ -2497,7 +2710,7 @@ SQLITE_API int sqlite3_changes(sqlite3*);
 ** To detect changes against a database file from other database
 ** connections use the [PRAGMA data_version] command or the
 ** [SQLITE_FCNTL_DATA_VERSION] [file control].
-** 
+**
 ** If a separate thread makes changes on the same database connection
 ** while [sqlite3_total_changes()] is running then the value
 ** returned is unpredictable and not meaningful.
@@ -2512,6 +2725,7 @@ SQLITE_API int sqlite3_changes(sqlite3*);
 ** </ul>
 */
 SQLITE_API int sqlite3_total_changes(sqlite3*);
+SQLITE_API sqlite3_int64 sqlite3_total_changes64(sqlite3*);
 
 /*
 ** CAPI3REF: Interrupt A Long-Running Query
@@ -2539,7 +2753,7 @@ SQLITE_API int sqlite3_total_changes(sqlite3*);
 **
 ** ^The sqlite3_interrupt(D) call is in effect until all currently running
 ** SQL statements on [database connection] D complete.  ^Any new SQL statements
-** that are started after the sqlite3_interrupt() call and before the 
+** that are started after the sqlite3_interrupt() call and before the
 ** running statement count reaches zero are interrupted as if they had been
 ** running prior to the sqlite3_interrupt() call.  ^New SQL statements
 ** that are started after the running statement count reaches zero are
@@ -2547,8 +2761,13 @@ SQLITE_API int sqlite3_total_changes(sqlite3*);
 ** ^A call to sqlite3_interrupt(D) that occurs when there are no running
 ** SQL statements is a no-op and has no effect on SQL statements
 ** that are started after the sqlite3_interrupt() call returns.
+**
+** ^The [sqlite3_is_interrupted(D)] interface can be used to determine whether
+** or not an interrupt is currently in effect for [database connection] D.
+** It returns 1 if an interrupt is currently in effect, or 0 otherwise.
 */
 SQLITE_API void sqlite3_interrupt(sqlite3*);
+SQLITE_API int sqlite3_is_interrupted(sqlite3*);
 
 /*
 ** CAPI3REF: Determine If An SQL Statement Is Complete
@@ -2571,7 +2790,7 @@ SQLITE_API void sqlite3_interrupt(sqlite3*);
 ** ^These routines do not parse the SQL statements thus
 ** will not detect syntactically incorrect SQL.
 **
-** ^(If SQLite has not been initialized using [sqlite3_initialize()] prior 
+** ^(If SQLite has not been initialized using [sqlite3_initialize()] prior
 ** to invoking sqlite3_complete16() then sqlite3_initialize() is invoked
 ** automatically by sqlite3_complete16().  If that initialization fails,
 ** then the return value from sqlite3_complete16() will be non-zero
@@ -2616,7 +2835,7 @@ SQLITE_API int sqlite3_complete16(const void *sql);
 ** The presence of a busy handler does not guarantee that it will be invoked
 ** when there is lock contention. ^If SQLite determines that invoking the busy
 ** handler could result in a deadlock, it will go ahead and return [SQLITE_BUSY]
-** to the application instead of invoking the 
+** to the application instead of invoking the
 ** busy handler.
 ** Consider a scenario where one process is holding a read lock that
 ** it is trying to promote to a reserved lock and
@@ -2641,7 +2860,7 @@ SQLITE_API int sqlite3_complete16(const void *sql);
 ** database connection that invoked the busy handler.  In other words,
 ** the busy handler is not reentrant.  Any such actions
 ** result in undefined behavior.
-** 
+**
 ** A busy handler must not close the database connection
 ** or [prepared statement] that invoked the busy handler.
 */
@@ -2759,7 +2978,7 @@ SQLITE_API void sqlite3_free_table(char **result);
 ** These routines are work-alikes of the "printf()" family of functions
 ** from the standard C library.
 ** These routines understand most of the common formatting options from
-** the standard library printf() 
+** the standard library printf()
 ** plus some additional non-standard formats ([%q], [%Q], [%w], and [%z]).
 ** See the [built-in printf()] documentation for details.
 **
@@ -2955,7 +3174,7 @@ SQLITE_API void sqlite3_randomness(int N, void *P);
 ** requested is ok.  ^When the callback returns [SQLITE_DENY], the
 ** [sqlite3_prepare_v2()] or equivalent call that triggered the
 ** authorizer will fail with an error message explaining that
-** access is denied. 
+** access is denied.
 **
 ** ^The first parameter to the authorizer callback is a copy of the third
 ** parameter to the sqlite3_set_authorizer() interface. ^The second parameter
@@ -3008,7 +3227,7 @@ SQLITE_API void sqlite3_randomness(int N, void *P);
 ** database connections for the meaning of "modify" in this paragraph.
 **
 ** ^When [sqlite3_prepare_v2()] is used to prepare a statement, the
-** statement might be re-prepared during [sqlite3_step()] due to a 
+** statement might be re-prepared during [sqlite3_step()] due to a
 ** schema change.  Hence, the application should ensure that the
 ** correct authorizer callback remains in place during the [sqlite3_step()].
 **
@@ -3095,8 +3314,8 @@ SQLITE_API int sqlite3_set_authorizer(
 #define SQLITE_RECURSIVE            33   /* NULL            NULL            */
 
 /*
-** CAPI3REF: Tracing And Profiling Functions
-** METHOD: sqlite3
+** CAPI3REF: Deprecated Tracing And Profiling Functions
+** DEPRECATED
 **
 ** These routines are deprecated. Use the [sqlite3_trace_v2()] interface
 ** instead of the routines described here.
@@ -3156,7 +3375,7 @@ SQLITE_API SQLITE_DEPRECATED void *sqlite3_profile(sqlite3*,
 ** execution of the prepared statement, such as at the start of each
 ** trigger subprogram. ^The P argument is a pointer to the
 ** [prepared statement]. ^The X argument is a pointer to a string which
-** is the unexpanded SQL text of the prepared statement or an SQL comment 
+** is the unexpanded SQL text of the prepared statement or an SQL comment
 ** that indicates the invocation of a trigger.  ^The callback can compute
 ** the same text that would have been returned by the legacy [sqlite3_trace()]
 ** interface by using the X argument when X begins with "--" and invoking
@@ -3166,13 +3385,13 @@ SQLITE_API SQLITE_DEPRECATED void *sqlite3_profile(sqlite3*,
 ** <dd>^An SQLITE_TRACE_PROFILE callback provides approximately the same
 ** information as is provided by the [sqlite3_profile()] callback.
 ** ^The P argument is a pointer to the [prepared statement] and the
-** X argument points to a 64-bit integer which is the estimated of
-** the number of nanosecond that the prepared statement took to run.
+** X argument points to a 64-bit integer which is approximately
+** the number of nanoseconds that the prepared statement took to run.
 ** ^The SQLITE_TRACE_PROFILE callback is invoked when the statement finishes.
 **
 ** [[SQLITE_TRACE_ROW]] <dt>SQLITE_TRACE_ROW</dt>
 ** <dd>^An SQLITE_TRACE_ROW callback is invoked whenever a prepared
-** statement generates a single row of result.  
+** statement generates a single row of result.
 ** ^The P argument is a pointer to the [prepared statement] and the
 ** X argument is unused.
 **
@@ -3199,10 +3418,12 @@ SQLITE_API SQLITE_DEPRECATED void *sqlite3_profile(sqlite3*,
 ** M argument should be the bitwise OR-ed combination of
 ** zero or more [SQLITE_TRACE] constants.
 **
-** ^Each call to either sqlite3_trace() or sqlite3_trace_v2() overrides 
-** (cancels) any prior calls to sqlite3_trace() or sqlite3_trace_v2().
+** ^Each call to either sqlite3_trace(D,X,P) or sqlite3_trace_v2(D,M,X,P)
+** overrides (cancels) all prior calls to sqlite3_trace(D,X,P) or
+** sqlite3_trace_v2(D,M,X,P) for the [database connection] D.  Each
+** database connection may have at most one trace callback.
 **
-** ^The X callback is invoked whenever any of the events identified by 
+** ^The X callback is invoked whenever any of the events identified by
 ** mask M occur.  ^The integer return value from the callback is currently
 ** ignored, though this may change in future releases.  Callback
 ** implementations should return zero to ensure future compatibility.
@@ -3230,12 +3451,12 @@ SQLITE_API int sqlite3_trace_v2(
 **
 ** ^The sqlite3_progress_handler(D,N,X,P) interface causes the callback
 ** function X to be invoked periodically during long running calls to
-** [sqlite3_exec()], [sqlite3_step()] and [sqlite3_get_table()] for
+** [sqlite3_step()] and [sqlite3_prepare()] and similar for
 ** database connection D.  An example use for this
 ** interface is to keep a GUI updated during a large query.
 **
-** ^The parameter P is passed through as the only parameter to the 
-** callback function X.  ^The parameter N is the approximate number of 
+** ^The parameter P is passed through as the only parameter to the
+** callback function X.  ^The parameter N is the approximate number of
 ** [virtual machine instructions] that are evaluated between successive
 ** invocations of the callback X.  ^If N is less than one then the progress
 ** handler is disabled.
@@ -3255,6 +3476,13 @@ SQLITE_API int sqlite3_trace_v2(
 ** Note that [sqlite3_prepare_v2()] and [sqlite3_step()] both modify their
 ** database connections for the meaning of "modify" in this paragraph.
 **
+** The progress handler callback would originally only be invoked from the
+** bytecode engine.  It still might be invoked during [sqlite3_prepare()]
+** and similar because those routines might force a reparse of the schema
+** which involves running the bytecode engine.  However, beginning with
+** SQLite version 3.41.0, the progress handler callback might also be
+** invoked directly from [sqlite3_prepare()] while analyzing and generating
+** code for complex queries.
 */
 SQLITE_API void sqlite3_progress_handler(sqlite3*, int, int(*)(void*), void*);
 
@@ -3262,7 +3490,7 @@ SQLITE_API void sqlite3_progress_handler(sqlite3*, int, int(*)(void*), void*);
 ** CAPI3REF: Opening A New Database Connection
 ** CONSTRUCTOR: sqlite3
 **
-** ^These routines open an SQLite database file as specified by the 
+** ^These routines open an SQLite database file as specified by the
 ** filename argument. ^The filename argument is interpreted as UTF-8 for
 ** sqlite3_open() and sqlite3_open_v2() and as UTF-16 in the native byte
 ** order for sqlite3_open16(). ^(A [database connection] handle is usually
@@ -3291,13 +3519,18 @@ SQLITE_API void sqlite3_progress_handler(sqlite3*, int, int(*)(void*), void*);
 **
 ** <dl>
 ** ^(<dt>[SQLITE_OPEN_READONLY]</dt>
-** <dd>The database is opened in read-only mode.  If the database does not
-** already exist, an error is returned.</dd>)^
+** <dd>The database is opened in read-only mode.  If the database does
+** not already exist, an error is returned.</dd>)^
 **
 ** ^(<dt>[SQLITE_OPEN_READWRITE]</dt>
-** <dd>The database is opened for reading and writing if possible, or reading
-** only if the file is write protected by the operating system.  In either
-** case the database must already exist, otherwise an error is returned.</dd>)^
+** <dd>The database is opened for reading and writing if possible, or
+** reading only if the file is write protected by the operating
+** system.  In either case the database must already exist, otherwise
+** an error is returned.  For historical reasons, if opening in
+** read-write mode fails due to OS-level permissions, an attempt is
+** made to open it in read-only mode. [sqlite3_db_readonly()] can be
+** used to determine whether the database is actually
+** read-write.</dd>)^
 **
 ** ^(<dt>[SQLITE_OPEN_READWRITE] | [SQLITE_OPEN_CREATE]</dt>
 ** <dd>The database is opened for reading and writing, and is created if
@@ -3335,20 +3568,39 @@ SQLITE_API void sqlite3_progress_handler(sqlite3*, int, int(*)(void*), void*);
 ** <dd>The database is opened [shared cache] enabled, overriding
 ** the default shared cache setting provided by
 ** [sqlite3_enable_shared_cache()].)^
+** The [use of shared cache mode is discouraged] and hence shared cache
+** capabilities may be omitted from many builds of SQLite.  In such cases,
+** this option is a no-op.
 **
 ** ^(<dt>[SQLITE_OPEN_PRIVATECACHE]</dt>
 ** <dd>The database is opened [shared cache] disabled, overriding
 ** the default shared cache setting provided by
 ** [sqlite3_enable_shared_cache()].)^
 **
+** [[OPEN_EXRESCODE]] ^(<dt>[SQLITE_OPEN_EXRESCODE]</dt>
+** <dd>The database connection comes up in "extended result code mode".
+** In other words, the database behaves as if
+** [sqlite3_extended_result_codes(db,1)] were called on the database
+** connection as soon as the connection is created. In addition to setting
+** the extended result code mode, this flag also causes [sqlite3_open_v2()]
+** to return an extended result code.</dd>
+**
 ** [[OPEN_NOFOLLOW]] ^(<dt>[SQLITE_OPEN_NOFOLLOW]</dt>
-** <dd>The database filename is not allowed to be a symbolic link</dd>
+** <dd>The database filename is not allowed to contain a symbolic link</dd>
 ** </dl>)^
 **
 ** If the 3rd parameter to sqlite3_open_v2() is not one of the
 ** required combinations shown above optionally combined with other
 ** [SQLITE_OPEN_READONLY | SQLITE_OPEN_* bits]
-** then the behavior is undefined.
+** then the behavior is undefined.  Historic versions of SQLite
+** have silently ignored surplus bits in the flags parameter to
+** sqlite3_open_v2(), however that behavior might not be carried through
+** into future versions of SQLite and so applications should not rely
+** upon it.  Note in particular that the SQLITE_OPEN_EXCLUSIVE flag is a no-op
+** for sqlite3_open_v2().  The SQLITE_OPEN_EXCLUSIVE does *not* cause
+** the open to fail if the database already exists.  The SQLITE_OPEN_EXCLUSIVE
+** flag is intended for use by the [sqlite3_vfs|VFS interface] only, and not
+** by sqlite3_open_v2().
 **
 ** ^The fourth parameter to sqlite3_open_v2() is the name of the
 ** [sqlite3_vfs] object that defines the operating system interface that
@@ -3381,17 +3633,17 @@ SQLITE_API void sqlite3_progress_handler(sqlite3*, int, int(*)(void*), void*);
 ** information.
 **
 ** URI filenames are parsed according to RFC 3986. ^If the URI contains an
-** authority, then it must be either an empty string or the string 
-** "localhost". ^If the authority is not an empty string or "localhost", an 
-** error is returned to the caller. ^The fragment component of a URI, if 
+** authority, then it must be either an empty string or the string
+** "localhost". ^If the authority is not an empty string or "localhost", an
+** error is returned to the caller. ^The fragment component of a URI, if
 ** present, is ignored.
 **
 ** ^SQLite uses the path component of the URI as the name of the disk file
-** which contains the database. ^If the path begins with a '/' character, 
-** then it is interpreted as an absolute path. ^If the path does not begin 
+** which contains the database. ^If the path begins with a '/' character,
+** then it is interpreted as an absolute path. ^If the path does not begin
 ** with a '/' (meaning that the authority section is omitted from the URI)
-** then the path is interpreted as a relative path. 
-** ^(On windows, the first component of an absolute path 
+** then the path is interpreted as a relative path.
+** ^(On windows, the first component of an absolute path
 ** is a drive specification (e.g. "C:").)^
 **
 ** [[core URI query parameters]]
@@ -3411,13 +3663,13 @@ SQLITE_API void sqlite3_progress_handler(sqlite3*, int, int(*)(void*), void*);
 **
 **   <li> <b>mode</b>: ^(The mode parameter may be set to either "ro", "rw",
 **     "rwc", or "memory". Attempting to set it to any other value is
-**     an error)^. 
-**     ^If "ro" is specified, then the database is opened for read-only 
-**     access, just as if the [SQLITE_OPEN_READONLY] flag had been set in the 
-**     third argument to sqlite3_open_v2(). ^If the mode option is set to 
-**     "rw", then the database is opened for read-write (but not create) 
-**     access, as if SQLITE_OPEN_READWRITE (but not SQLITE_OPEN_CREATE) had 
-**     been set. ^Value "rwc" is equivalent to setting both 
+**     an error)^.
+**     ^If "ro" is specified, then the database is opened for read-only
+**     access, just as if the [SQLITE_OPEN_READONLY] flag had been set in the
+**     third argument to sqlite3_open_v2(). ^If the mode option is set to
+**     "rw", then the database is opened for read-write (but not create)
+**     access, as if SQLITE_OPEN_READWRITE (but not SQLITE_OPEN_CREATE) had
+**     been set. ^Value "rwc" is equivalent to setting both
 **     SQLITE_OPEN_READWRITE and SQLITE_OPEN_CREATE.  ^If the mode option is
 **     set to "memory" then a pure [in-memory database] that never reads
 **     or writes from disk is used. ^It is an error to specify a value for
@@ -3427,7 +3679,7 @@ SQLITE_API void sqlite3_progress_handler(sqlite3*, int, int(*)(void*), void*);
 **   <li> <b>cache</b>: ^The cache parameter may be set to either "shared" or
 **     "private". ^Setting it to "shared" is equivalent to setting the
 **     SQLITE_OPEN_SHAREDCACHE bit in the flags argument passed to
-**     sqlite3_open_v2(). ^Setting the cache parameter to "private" is 
+**     sqlite3_open_v2(). ^Setting the cache parameter to "private" is
 **     equivalent to setting the SQLITE_OPEN_PRIVATECACHE bit.
 **     ^If sqlite3_open_v2() is used and the "cache" parameter is present in
 **     a URI filename, its value overrides any behavior requested by setting
@@ -3453,7 +3705,7 @@ SQLITE_API void sqlite3_progress_handler(sqlite3*, int, int(*)(void*), void*);
 **     property on a database file that does in fact change can result
 **     in incorrect query results and/or [SQLITE_CORRUPT] errors.
 **     See also: [SQLITE_IOCAP_IMMUTABLE].
-**       
+**
 ** </ul>
 **
 ** ^Specifying an unknown parameter in the query component of a URI is not an
@@ -3465,36 +3717,37 @@ SQLITE_API void sqlite3_progress_handler(sqlite3*, int, int(*)(void*), void*);
 **
 ** <table border="1" align=center cellpadding=5>
 ** <tr><th> URI filenames <th> Results
-** <tr><td> file:data.db <td> 
+** <tr><td> file:data.db <td>
 **          Open the file "data.db" in the current directory.
 ** <tr><td> file:[REDACTED_USER_HOME]/data.db<br>
-**          file://[REDACTED_USER_HOME]/data.db <br> 
-**          file://localhost[REDACTED_USER_HOME]/data.db <br> <td> 
+**          file://[REDACTED_USER_HOME]/data.db <br>
+**          file://localhost[REDACTED_USER_HOME]/data.db <br> <td>
 **          Open the database file "[REDACTED_USER_HOME]/data.db".
-** <tr><td> file://darkstar[REDACTED_USER_HOME]/data.db <td> 
+** <tr><td> file://darkstar[REDACTED_USER_HOME]/data.db <td>
 **          An error. "darkstar" is not a recognized authority.
-** <tr><td style="white-space:nowrap"> 
+** <tr><td style="white-space:nowrap">
 **          file:///C:/Documents%20and%20Settings/fred/Desktop/data.db
 **     <td> Windows only: Open the file "data.db" on fred's desktop on drive
-**          C:. Note that the %20 escaping in this example is not strictly 
+**          C:. Note that the %20 escaping in this example is not strictly
 **          necessary - space characters can be used literally
 **          in URI filenames.
-** <tr><td> file:data.db?mode=ro&cache=private <td> 
+** <tr><td> file:data.db?mode=ro&cache=private <td>
 **          Open file "data.db" in the current directory for read-only access.
 **          Regardless of whether or not shared-cache mode is enabled by
 **          default, use a private cache.
 ** <tr><td> file:[REDACTED_USER_HOME]/data.db?vfs=unix-dotfile <td>
 **          Open file "[REDACTED_USER_HOME]/data.db". Use the special VFS "unix-dotfile"
 **          that uses dot-files in place of posix advisory locking.
-** <tr><td> file:data.db?mode=readonly <td> 
+** <tr><td> file:data.db?mode=readonly <td>
 **          An error. "readonly" is not a valid option for the "mode" parameter.
+**          Use "ro" instead:  "file:data.db?mode=ro".
 ** </table>
 **
 ** ^URI hexadecimal escape sequences (%HH) are supported within the path and
 ** query components of a URI. A hexadecimal escape sequence consists of a
-** percent sign - "%" - followed by exactly two hexadecimal digits 
+** percent sign - "%" - followed by exactly two hexadecimal digits
 ** specifying an octet value. ^Before the path or query components of a
-** URI filename are interpreted, they are encoded using UTF-8 and all 
+** URI filename are interpreted, they are encoded using UTF-8 and all
 ** hexadecimal escape sequences replaced by a single byte containing the
 ** corresponding octet. If this process generates an invalid UTF-8 encoding,
 ** the results are undefined.
@@ -3530,14 +3783,25 @@ SQLITE_API int sqlite3_open_v2(
 ** CAPI3REF: Obtain Values For URI Parameters
 **
 ** These are utility routines, useful to [VFS|custom VFS implementations],
-** that check if a database file was a URI that contained a specific query 
+** that check if a database file was a URI that contained a specific query
 ** parameter, and if so obtains the value of that query parameter.
 **
-** If F is the database filename pointer passed into the xOpen() method of 
-** a VFS implementation or it is the return value of [sqlite3_db_filename()]
+** The first parameter to these interfaces (hereafter referred to
+** as F) must be one of:
+** <ul>
+** <li> A database filename pointer created by the SQLite core and
+** passed into the xOpen() method of a VFS implementation, or
+** <li> A filename obtained from [sqlite3_db_filename()], or
+** <li> A new filename constructed using [sqlite3_create_filename()].
+** </ul>
+** If the F parameter is not one of the above, then the behavior is
+** undefined and probably undesirable.  Older versions of SQLite were
+** more tolerant of invalid F parameters than newer versions.
+**
+** If F is a suitable filename (as described in the previous paragraph)
 ** and if P is the name of the query parameter, then
 ** sqlite3_uri_parameter(F,P) returns the value of the P
-** parameter if it exists or a NULL pointer if P does not appear as a 
+** parameter if it exists or a NULL pointer if P does not appear as a
 ** query parameter on F.  If P is a query parameter of F and it
 ** has no explicit value, then sqlite3_uri_parameter(F,P) returns
 ** a pointer to an empty string.
@@ -3546,7 +3810,7 @@ SQLITE_API int sqlite3_open_v2(
 ** parameter and returns true (1) or false (0) according to the value
 ** of P.  The sqlite3_uri_boolean(F,P,B) routine returns true (1) if the
 ** value of query parameter P is one of "yes", "true", or "on" in any
-** case or if the value begins with a non-zero number.  The 
+** case or if the value begins with a non-zero number.  The
 ** sqlite3_uri_boolean(F,P,B) routines returns false (0) if the value of
 ** query parameter P is one of "no", "false", or "off" in any case or
 ** if the value begins with a numeric zero.  If P is not a query
@@ -3564,7 +3828,7 @@ SQLITE_API int sqlite3_open_v2(
 ** parameters minus 1.  The N value is zero-based so N should be 0 to obtain
 ** the name of the first query parameter, 1 for the second parameter, and
 ** so forth.
-** 
+**
 ** If F is a NULL pointer, then sqlite3_uri_parameter(F,P) returns NULL and
 ** sqlite3_uri_boolean(F,P,B) returns B.  If F is not a NULL pointer and
 ** is not a database file pathname pointer that the SQLite core passed
@@ -3581,10 +3845,10 @@ SQLITE_API int sqlite3_open_v2(
 **
 ** See the [URI filename] documentation for additional information.
 */
-SQLITE_API const char *sqlite3_uri_parameter(const char *zFilename, const char *zParam);
-SQLITE_API int sqlite3_uri_boolean(const char *zFile, const char *zParam, int bDefault);
-SQLITE_API sqlite3_int64 sqlite3_uri_int64(const char*, const char*, sqlite3_int64);
-SQLITE_API const char *sqlite3_uri_key(const char *zFilename, int N);
+SQLITE_API const char *sqlite3_uri_parameter(sqlite3_filename z, const char *zParam);
+SQLITE_API int sqlite3_uri_boolean(sqlite3_filename z, const char *zParam, int bDefault);
+SQLITE_API sqlite3_int64 sqlite3_uri_int64(sqlite3_filename, const char*, sqlite3_int64);
+SQLITE_API const char *sqlite3_uri_key(sqlite3_filename z, int N);
 
 /*
 ** CAPI3REF:  Translate filenames
@@ -3613,21 +3877,93 @@ SQLITE_API const char *sqlite3_uri_key(const char *zFilename, int N);
 ** return value from [sqlite3_db_filename()], then the result is
 ** undefined and is likely a memory access violation.
 */
-SQLITE_API const char *sqlite3_filename_database(const char*);
-SQLITE_API const char *sqlite3_filename_journal(const char*);
-SQLITE_API const char *sqlite3_filename_wal(const char*);
+SQLITE_API const char *sqlite3_filename_database(sqlite3_filename);
+SQLITE_API const char *sqlite3_filename_journal(sqlite3_filename);
+SQLITE_API const char *sqlite3_filename_wal(sqlite3_filename);
 
+/*
+** CAPI3REF:  Database File Corresponding To A Journal
+**
+** ^If X is the name of a rollback or WAL-mode journal file that is
+** passed into the xOpen method of [sqlite3_vfs], then
+** sqlite3_database_file_object(X) returns a pointer to the [sqlite3_file]
+** object that represents the main database file.
+**
+** This routine is intended for use in custom [VFS] implementations
+** only.  It is not a general-purpose interface.
+** The argument sqlite3_file_object(X) must be a filename pointer that
+** has been passed into [sqlite3_vfs].xOpen method where the
+** flags parameter to xOpen contains one of the bits
+** [SQLITE_OPEN_MAIN_JOURNAL] or [SQLITE_OPEN_WAL].  Any other use
+** of this routine results in undefined and probably undesirable
+** behavior.
+*/
+SQLITE_API sqlite3_file *sqlite3_database_file_object(const char*);
+
+/*
+** CAPI3REF: Create and Destroy VFS Filenames
+**
+** These interfaces are provided for use by [VFS shim] implementations and
+** are not useful outside of that context.
+**
+** The sqlite3_create_filename(D,J,W,N,P) allocates memory to hold a version of
+** database filename D with corresponding journal file J and WAL file W and
+** with N URI parameters key/values pairs in the array P.  The result from
+** sqlite3_create_filename(D,J,W,N,P) is a pointer to a database filename that
+** is safe to pass to routines like:
+** <ul>
+** <li> [sqlite3_uri_parameter()],
+** <li> [sqlite3_uri_boolean()],
+** <li> [sqlite3_uri_int64()],
+** <li> [sqlite3_uri_key()],
+** <li> [sqlite3_filename_database()],
+** <li> [sqlite3_filename_journal()], or
+** <li> [sqlite3_filename_wal()].
+** </ul>
+** If a memory allocation error occurs, sqlite3_create_filename() might
+** return a NULL pointer.  The memory obtained from sqlite3_create_filename(X)
+** must be released by a corresponding call to sqlite3_free_filename(Y).
+**
+** The P parameter in sqlite3_create_filename(D,J,W,N,P) should be an array
+** of 2*N pointers to strings.  Each pair of pointers in this array corresponds
+** to a key and value for a query parameter.  The P parameter may be a NULL
+** pointer if N is zero.  None of the 2*N pointers in the P array may be
+** NULL pointers and key pointers should not be empty strings.
+** None of the D, J, or W parameters to sqlite3_create_filename(D,J,W,N,P) may
+** be NULL pointers, though they can be empty strings.
+**
+** The sqlite3_free_filename(Y) routine releases a memory allocation
+** previously obtained from sqlite3_create_filename().  Invoking
+** sqlite3_free_filename(Y) where Y is a NULL pointer is a harmless no-op.
+**
+** If the Y parameter to sqlite3_free_filename(Y) is anything other
+** than a NULL pointer or a pointer previously acquired from
+** sqlite3_create_filename(), then bad things such as heap
+** corruption or segfaults may occur. The value Y should not be
+** used again after sqlite3_free_filename(Y) has been called.  This means
+** that if the [sqlite3_vfs.xOpen()] method of a VFS has been called using Y,
+** then the corresponding [sqlite3_module.xClose() method should also be
+** invoked prior to calling sqlite3_free_filename(Y).
+*/
+SQLITE_API sqlite3_filename sqlite3_create_filename(
+  const char *zDatabase,
+  const char *zJournal,
+  const char *zWal,
+  int nParam,
+  const char **azParam
+);
+SQLITE_API void sqlite3_free_filename(sqlite3_filename);
 
 /*
 ** CAPI3REF: Error Codes And Messages
 ** METHOD: sqlite3
 **
-** ^If the most recent sqlite3_* API call associated with 
+** ^If the most recent sqlite3_* API call associated with
 ** [database connection] D failed, then the sqlite3_errcode(D) interface
 ** returns the numeric [result code] or [extended result code] for that
 ** API call.
 ** ^The sqlite3_extended_errcode()
-** interface is the same except that it always returns the 
+** interface is the same except that it always returns the
 ** [extended result code] even when extended result codes are
 ** disabled.
 **
@@ -3635,27 +3971,38 @@ SQLITE_API const char *sqlite3_filename_wal(const char*);
 ** sqlite3_extended_errcode() might change with each API call.
 ** Except, there are some interfaces that are guaranteed to never
 ** change the value of the error code.  The error-code preserving
-** interfaces are:
+** interfaces include the following:
 **
 ** <ul>
 ** <li> sqlite3_errcode()
 ** <li> sqlite3_extended_errcode()
 ** <li> sqlite3_errmsg()
 ** <li> sqlite3_errmsg16()
+** <li> sqlite3_error_offset()
 ** </ul>
 **
 ** ^The sqlite3_errmsg() and sqlite3_errmsg16() return English-language
-** text that describes the error, as either UTF-8 or UTF-16 respectively.
+** text that describes the error, as either UTF-8 or UTF-16 respectively,
+** or NULL if no error message is available.
+** (See how SQLite handles [invalid UTF] for exceptions to this rule.)
 ** ^(Memory to hold the error message string is managed internally.
 ** The application does not need to worry about freeing the result.
 ** However, the error string might be overwritten or deallocated by
 ** subsequent calls to other SQLite interface functions.)^
 **
-** ^The sqlite3_errstr() interface returns the English-language text
-** that describes the [result code], as UTF-8.
+** ^The sqlite3_errstr(E) interface returns the English-language text
+** that describes the [result code] E, as UTF-8, or NULL if E is not an
+** result code for which a text error message is available.
 ** ^(Memory to hold the error message string is managed internally
 ** and must not be freed by the application)^.
 **
+** ^If the most recent error references a specific token in the input
+** SQL, the sqlite3_error_offset() interface returns the byte offset
+** of the start of that token.  ^The byte offset returned by
+** sqlite3_error_offset() assumes that the input SQL is UTF8.
+** ^If the most recent error does not reference a specific token in the input
+** SQL, then the sqlite3_error_offset() function returns -1.
+**
 ** When the serialized [threading mode] is in use, it might be the
 ** case that a second error occurs on a separate thread in between
 ** the time of the first error and the call to these interfaces.
@@ -3675,6 +4022,7 @@ SQLITE_API int sqlite3_extended_errcode(sqlite3 *db);
 SQLITE_API const char *sqlite3_errmsg(sqlite3*);
 SQLITE_API const void *sqlite3_errmsg16(sqlite3*);
 SQLITE_API const char *sqlite3_errstr(int);
+SQLITE_API int sqlite3_error_offset(sqlite3 *db);
 
 /*
 ** CAPI3REF: Prepared Statement Object
@@ -3684,7 +4032,7 @@ SQLITE_API const char *sqlite3_errstr(int);
 ** has been compiled into binary form and is ready to be evaluated.
 **
 ** Think of each SQL statement as a separate computer program.  The
-** original SQL text is source code.  A prepared statement object 
+** original SQL text is source code.  A prepared statement object
 ** is the compiled object code.  All SQL must be converted into a
 ** prepared statement before it can be run.
 **
@@ -3714,7 +4062,7 @@ typedef struct sqlite3_stmt sqlite3_stmt;
 ** new limit for that construct.)^
 **
 ** ^If the new limit is a negative number, the limit is unchanged.
-** ^(For each limit category SQLITE_LIMIT_<i>NAME</i> there is a 
+** ^(For each limit category SQLITE_LIMIT_<i>NAME</i> there is a
 ** [limits | hard upper bound]
 ** set at compile-time by a C preprocessor macro called
 ** [limits | SQLITE_MAX_<i>NAME</i>].
@@ -3722,7 +4070,7 @@ typedef struct sqlite3_stmt sqlite3_stmt;
 ** ^Attempts to increase a limit above its hard upper bound are
 ** silently truncated to the hard upper bound.
 **
-** ^Regardless of whether or not the limit was changed, the 
+** ^Regardless of whether or not the limit was changed, the
 ** [sqlite3_limit()] interface returns the prior value of the limit.
 ** ^Hence, to find the current value of a limit without changing it,
 ** simply invoke this interface with the third parameter set to -1.
@@ -3827,7 +4175,7 @@ SQLITE_API int sqlite3_limit(sqlite3*, int id, int newVal);
 ** <dd>The SQLITE_PREPARE_PERSISTENT flag is a hint to the query planner
 ** that the prepared statement will be retained for a long time and
 ** probably reused many times.)^ ^Without this flag, [sqlite3_prepare_v3()]
-** and [sqlite3_prepare16_v3()] assume that the prepared statement will 
+** and [sqlite3_prepare16_v3()] assume that the prepared statement will
 ** be used just once or at most a few times and then destroyed using
 ** [sqlite3_finalize()] relatively soon. The current implementation acts
 ** on this hint by avoiding the use of [lookaside memory] so as not to
@@ -3883,13 +4231,17 @@ SQLITE_API int sqlite3_limit(sqlite3*, int id, int newVal);
 ** and sqlite3_prepare16_v3() use UTF-16.
 **
 ** ^If the nByte argument is negative, then zSql is read up to the
-** first zero terminator. ^If nByte is positive, then it is the
-** number of bytes read from zSql.  ^If nByte is zero, then no prepared
+** first zero terminator. ^If nByte is positive, then it is the maximum
+** number of bytes read from zSql.  When nByte is positive, zSql is read
+** up to the first zero terminator or until the nByte bytes have been read,
+** whichever comes first.  ^If nByte is zero, then no prepared
 ** statement is generated.
 ** If the caller knows that the supplied string is nul-terminated, then
 ** there is a small performance advantage to passing an nByte parameter that
 ** is the number of bytes in the input string <i>including</i>
 ** the nul-terminator.
+** Note that nByte measure the length of the input in bytes, not
+** characters, even for the UTF-16 interfaces.
 **
 ** ^If pzTail is not NULL then *pzTail is made to point to the first byte
 ** past the end of the first SQL statement in zSql.  These routines only
@@ -3934,12 +4286,12 @@ SQLITE_API int sqlite3_limit(sqlite3*, int id, int newVal);
 ** </li>
 **
 ** <li>
-** ^If the specific value bound to a [parameter | host parameter] in the 
+** ^If the specific value bound to a [parameter | host parameter] in the
 ** WHERE clause might influence the choice of query plan for a statement,
-** then the statement will be automatically recompiled, as if there had been 
+** then the statement will be automatically recompiled, as if there had been
 ** a schema change, on the first [sqlite3_step()] call following any change
-** to the [sqlite3_bind_text | bindings] of that [parameter]. 
-** ^The specific value of a WHERE-clause [parameter] might influence the 
+** to the [sqlite3_bind_text | bindings] of that [parameter].
+** ^The specific value of a WHERE-clause [parameter] might influence the
 ** choice of query plan if the parameter is the left-hand side of a [LIKE]
 ** or [GLOB] operator or if the parameter is compared to an indexed column
 ** and the [SQLITE_ENABLE_STAT4] compile-time option is enabled.
@@ -4032,12 +4384,17 @@ SQLITE_API int sqlite3_prepare16_v3(
 ** are managed by SQLite and are automatically freed when the prepared
 ** statement is finalized.
 ** ^The string returned by sqlite3_expanded_sql(P), on the other hand,
-** is obtained from [sqlite3_malloc()] and must be free by the application
+** is obtained from [sqlite3_malloc()] and must be freed by the application
 ** by passing it to [sqlite3_free()].
+**
+** ^The sqlite3_normalized_sql() interface is only available if
+** the [SQLITE_ENABLE_NORMALIZE] compile-time option is defined.
 */
 SQLITE_API const char *sqlite3_sql(sqlite3_stmt *pStmt);
 SQLITE_API char *sqlite3_expanded_sql(sqlite3_stmt *pStmt);
+#ifdef SQLITE_ENABLE_NORMALIZE
 SQLITE_API const char *sqlite3_normalized_sql(sqlite3_stmt *pStmt);
+#endif
 
 /*
 ** CAPI3REF: Determine If An SQL Statement Writes The Database
@@ -4048,8 +4405,8 @@ SQLITE_API const char *sqlite3_normalized_sql(sqlite3_stmt *pStmt);
 ** the content of the database file.
 **
 ** Note that [application-defined SQL functions] or
-** [virtual tables] might change the database indirectly as a side effect.  
-** ^(For example, if an application defines a function "eval()" that 
+** [virtual tables] might change the database indirectly as a side effect.
+** ^(For example, if an application defines a function "eval()" that
 ** calls [sqlite3_exec()], then the following SQL statement would
 ** change the database file through side-effects:
 **
@@ -4063,15 +4420,28 @@ SQLITE_API const char *sqlite3_normalized_sql(sqlite3_stmt *pStmt);
 ** ^Transaction control statements such as [BEGIN], [COMMIT], [ROLLBACK],
 ** [SAVEPOINT], and [RELEASE] cause sqlite3_stmt_readonly() to return true,
 ** since the statements themselves do not actually modify the database but
-** rather they control the timing of when other statements modify the 
+** rather they control the timing of when other statements modify the
 ** database.  ^The [ATTACH] and [DETACH] statements also cause
 ** sqlite3_stmt_readonly() to return true since, while those statements
-** change the configuration of a database connection, they do not make 
+** change the configuration of a database connection, they do not make
 ** changes to the content of the database files on disk.
 ** ^The sqlite3_stmt_readonly() interface returns true for [BEGIN] since
 ** [BEGIN] merely sets internal flags, but the [BEGIN|BEGIN IMMEDIATE] and
 ** [BEGIN|BEGIN EXCLUSIVE] commands do touch the database and so
 ** sqlite3_stmt_readonly() returns false for those commands.
+**
+** ^This routine returns false if there is any possibility that the
+** statement might change the database file.  ^A false return does
+** not guarantee that the statement will change the database file.
+** ^For example, an UPDATE statement might have a WHERE clause that
+** makes it a no-op, but the sqlite3_stmt_readonly() result would still
+** be false.  ^Similarly, a CREATE TABLE IF NOT EXISTS statement is a
+** read-only no-op if the table already exists, but
+** sqlite3_stmt_readonly() still returns false for such a statement.
+**
+** ^If prepared statement X is an [EXPLAIN] or [EXPLAIN QUERY PLAN]
+** statement, then sqlite3_stmt_readonly(X) returns the same value as
+** if the EXPLAIN or EXPLAIN QUERY PLAN prefix were omitted.
 */
 SQLITE_API int sqlite3_stmt_readonly(sqlite3_stmt *pStmt);
 
@@ -4087,23 +4457,58 @@ SQLITE_API int sqlite3_stmt_readonly(sqlite3_stmt *pStmt);
 */
 SQLITE_API int sqlite3_stmt_isexplain(sqlite3_stmt *pStmt);
 
+/*
+** CAPI3REF: Change The EXPLAIN Setting For A Prepared Statement
+** METHOD: sqlite3_stmt
+**
+** The sqlite3_stmt_explain(S,E) interface changes the EXPLAIN
+** setting for [prepared statement] S.  If E is zero, then S becomes
+** a normal prepared statement.  If E is 1, then S behaves as if
+** its SQL text began with "[EXPLAIN]".  If E is 2, then S behaves as if
+** its SQL text began with "[EXPLAIN QUERY PLAN]".
+**
+** Calling sqlite3_stmt_explain(S,E) might cause S to be reprepared.
+** SQLite tries to avoid a reprepare, but a reprepare might be necessary
+** on the first transition into EXPLAIN or EXPLAIN QUERY PLAN mode.
+**
+** Because of the potential need to reprepare, a call to
+** sqlite3_stmt_explain(S,E) will fail with SQLITE_ERROR if S cannot be
+** reprepared because it was created using [sqlite3_prepare()] instead of
+** the newer [sqlite3_prepare_v2()] or [sqlite3_prepare_v3()] interfaces and
+** hence has no saved SQL text with which to reprepare.
+**
+** Changing the explain setting for a prepared statement does not change
+** the original SQL text for the statement.  Hence, if the SQL text originally
+** began with EXPLAIN or EXPLAIN QUERY PLAN, but sqlite3_stmt_explain(S,0)
+** is called to convert the statement into an ordinary statement, the EXPLAIN
+** or EXPLAIN QUERY PLAN keywords will still appear in the sqlite3_sql(S)
+** output, even though the statement now acts like a normal SQL statement.
+**
+** This routine returns SQLITE_OK if the explain mode is successfully
+** changed, or an error code if the explain mode could not be changed.
+** The explain mode cannot be changed while a statement is active.
+** Hence, it is good practice to call [sqlite3_reset(S)]
+** immediately prior to calling sqlite3_stmt_explain(S,E).
+*/
+SQLITE_API int sqlite3_stmt_explain(sqlite3_stmt *pStmt, int eMode);
+
 /*
 ** CAPI3REF: Determine If A Prepared Statement Has Been Reset
 ** METHOD: sqlite3_stmt
 **
 ** ^The sqlite3_stmt_busy(S) interface returns true (non-zero) if the
-** [prepared statement] S has been stepped at least once using 
+** [prepared statement] S has been stepped at least once using
 ** [sqlite3_step(S)] but has neither run to completion (returned
 ** [SQLITE_DONE] from [sqlite3_step(S)]) nor
 ** been reset using [sqlite3_reset(S)].  ^The sqlite3_stmt_busy(S)
-** interface returns false if S is a NULL pointer.  If S is not a 
+** interface returns false if S is a NULL pointer.  If S is not a
 ** NULL pointer and is not a pointer to a valid [prepared statement]
 ** object, then the behavior is undefined and probably undesirable.
 **
 ** This interface can be used in combination [sqlite3_next_stmt()]
-** to locate all prepared statements associated with a database 
+** to locate all prepared statements associated with a database
 ** connection that are in need of being reset.  This can be used,
-** for example, in diagnostic routines to search for prepared 
+** for example, in diagnostic routines to search for prepared
 ** statements that are holding a transaction open.
 */
 SQLITE_API int sqlite3_stmt_busy(sqlite3_stmt*);
@@ -4122,7 +4527,7 @@ SQLITE_API int sqlite3_stmt_busy(sqlite3_stmt*);
 ** will accept either a protected or an unprotected sqlite3_value.
 ** Every interface that accepts sqlite3_value arguments specifies
 ** whether or not it requires a protected sqlite3_value.  The
-** [sqlite3_value_dup()] interface can be used to construct a new 
+** [sqlite3_value_dup()] interface can be used to construct a new
 ** protected sqlite3_value from an unprotected sqlite3_value.
 **
 ** The terms "protected" and "unprotected" refer to whether or not
@@ -4130,7 +4535,7 @@ SQLITE_API int sqlite3_stmt_busy(sqlite3_stmt*);
 ** sqlite3_value object but no mutex is held for an unprotected
 ** sqlite3_value object.  If SQLite is compiled to be single-threaded
 ** (with [SQLITE_THREADSAFE=0] and with [sqlite3_threadsafe()] returning 0)
-** or if SQLite is run in one of reduced mutex modes 
+** or if SQLite is run in one of reduced mutex modes
 ** [SQLITE_CONFIG_SINGLETHREAD] or [SQLITE_CONFIG_MULTITHREAD]
 ** then there is no distinction between protected and unprotected
 ** sqlite3_value objects and they can be used interchangeably.  However,
@@ -4140,6 +4545,8 @@ SQLITE_API int sqlite3_stmt_busy(sqlite3_stmt*);
 **
 ** ^The sqlite3_value objects that are passed as parameters into the
 ** implementation of [application-defined SQL functions] are protected.
+** ^The sqlite3_value objects returned by [sqlite3_vtab_rhs_value()]
+** are protected.
 ** ^The sqlite3_value object returned by
 ** [sqlite3_column_value()] is unprotected.
 ** Unprotected sqlite3_value objects may only be used as arguments
@@ -4199,12 +4606,30 @@ typedef struct sqlite3_context sqlite3_context;
 ** [sqlite3_bind_parameter_index()] API if desired.  ^The index
 ** for "?NNN" parameters is the value of NNN.
 ** ^The NNN value must be between 1 and the [sqlite3_limit()]
-** parameter [SQLITE_LIMIT_VARIABLE_NUMBER] (default value: 999).
+** parameter [SQLITE_LIMIT_VARIABLE_NUMBER] (default value: 32766).
 **
 ** ^The third argument is the value to bind to the parameter.
 ** ^If the third parameter to sqlite3_bind_text() or sqlite3_bind_text16()
 ** or sqlite3_bind_blob() is a NULL pointer then the fourth parameter
 ** is ignored and the end result is the same as sqlite3_bind_null().
+** ^If the third parameter to sqlite3_bind_text() is not NULL, then
+** it should be a pointer to well-formed UTF8 text.
+** ^If the third parameter to sqlite3_bind_text16() is not NULL, then
+** it should be a pointer to well-formed UTF16 text.
+** ^If the third parameter to sqlite3_bind_text64() is not NULL, then
+** it should be a pointer to a well-formed unicode string that is
+** either UTF8 if the sixth parameter is SQLITE_UTF8, or UTF16
+** otherwise.
+**
+** [[byte-order determination rules]] ^The byte-order of
+** UTF16 input text is determined by the byte-order mark (BOM, U+FEFF)
+** found in first character, which is removed, or in the absence of a BOM
+** the byte order is the native byte order of the host
+** machine for sqlite3_bind_text16() or the byte order specified in
+** the 6th parameter for sqlite3_bind_text64().)^
+** ^If UTF16 input text contains invalid unicode
+** characters, then SQLite might change those invalid characters
+** into the unicode replacement character: U+FFFD.
 **
 ** ^(In those routines that have a fourth argument, its value is the
 ** number of bytes in the parameter.  To be clear: the value is the
@@ -4218,23 +4643,27 @@ typedef struct sqlite3_context sqlite3_context;
 ** or sqlite3_bind_text16() or sqlite3_bind_text64() then
 ** that parameter must be the byte offset
 ** where the NUL terminator would occur assuming the string were NUL
-** terminated.  If any NUL characters occur at byte offsets less than 
+** terminated.  If any NUL characters occurs at byte offsets less than
 ** the value of the fourth parameter then the resulting string value will
 ** contain embedded NULs.  The result of expressions involving strings
 ** with embedded NULs is undefined.
 **
-** ^The fifth argument to the BLOB and string binding interfaces
-** is a destructor used to dispose of the BLOB or
-** string after SQLite has finished with it.  ^The destructor is called
-** to dispose of the BLOB or string even if the call to the bind API fails,
-** except the destructor is not called if the third parameter is a NULL
-** pointer or the fourth parameter is negative.
-** ^If the fifth argument is
-** the special value [SQLITE_STATIC], then SQLite assumes that the
-** information is in static, unmanaged space and does not need to be freed.
-** ^If the fifth argument has the value [SQLITE_TRANSIENT], then
-** SQLite makes its own private copy of the data immediately, before
-** the sqlite3_bind_*() routine returns.
+** ^The fifth argument to the BLOB and string binding interfaces controls
+** or indicates the lifetime of the object referenced by the third parameter.
+** These three options exist:
+** ^ (1) A destructor to dispose of the BLOB or string after SQLite has finished
+** with it may be passed. ^It is called to dispose of the BLOB or string even
+** if the call to the bind API fails, except the destructor is not called if
+** the third parameter is a NULL pointer or the fourth parameter is negative.
+** ^ (2) The special constant, [SQLITE_STATIC], may be passed to indicate that
+** the application remains responsible for disposing of the object. ^In this
+** case, the object and the provided pointer to it must remain valid until
+** either the prepared statement is finalized or the same SQL parameter is
+** bound to something else, whichever occurs sooner.
+** ^ (3) The constant, [SQLITE_TRANSIENT], may be passed to indicate that the
+** object is to be copied prior to the return from sqlite3_bind_*(). ^The
+** object and pointer to it must remain valid until then. ^SQLite will then
+** manage the lifetime of its private copy.
 **
 ** ^The sixth argument to sqlite3_bind_text64() must be one of
 ** [SQLITE_UTF8], [SQLITE_UTF16], [SQLITE_UTF16BE], or [SQLITE_UTF16LE]
@@ -4380,7 +4809,7 @@ SQLITE_API int sqlite3_clear_bindings(sqlite3_stmt*);
 ** METHOD: sqlite3_stmt
 **
 ** ^Return the number of columns in the result set returned by the
-** [prepared statement]. ^If this routine returns 0, that means the 
+** [prepared statement]. ^If this routine returns 0, that means the
 ** [prepared statement] returns no data (for example an [UPDATE]).
 ** ^However, just because this routine returns a positive number does not
 ** mean that one or more rows of data will be returned.  ^A SELECT statement
@@ -4562,7 +4991,7 @@ SQLITE_API const void *sqlite3_column_decltype16(sqlite3_stmt*,int);
 ** For all versions of SQLite up to and including 3.6.23.1, a call to
 ** [sqlite3_reset()] was required after sqlite3_step() returned anything
 ** other than [SQLITE_ROW] before any subsequent invocation of
-** sqlite3_step().  Failure to reset the prepared statement using 
+** sqlite3_step().  Failure to reset the prepared statement using
 ** [sqlite3_reset()] would result in an [SQLITE_MISUSE] return from
 ** sqlite3_step().  But after [version 3.6.23.1] ([dateof:3.6.23.1],
 ** sqlite3_step() began
@@ -4653,7 +5082,7 @@ SQLITE_API int sqlite3_data_count(sqlite3_stmt *pStmt);
 ** <tr><td><b>sqlite3_column_int64</b><td>&rarr;<td>64-bit INTEGER result
 ** <tr><td><b>sqlite3_column_text</b><td>&rarr;<td>UTF-8 TEXT result
 ** <tr><td><b>sqlite3_column_text16</b><td>&rarr;<td>UTF-16 TEXT result
-** <tr><td><b>sqlite3_column_value</b><td>&rarr;<td>The result as an 
+** <tr><td><b>sqlite3_column_value</b><td>&rarr;<td>The result as an
 ** [sqlite3_value|unprotected sqlite3_value] object.
 ** <tr><td>&nbsp;<td>&nbsp;<td>&nbsp;
 ** <tr><td><b>sqlite3_column_bytes</b><td>&rarr;<td>Size of a BLOB
@@ -4701,7 +5130,7 @@ SQLITE_API int sqlite3_data_count(sqlite3_stmt *pStmt);
 ** The return value of sqlite3_column_type() can be used to decide which
 ** of the first six interface should be used to extract the column value.
 ** The value returned by sqlite3_column_type() is only meaningful if no
-** automatic type conversions have occurred for the value in question.  
+** automatic type conversions have occurred for the value in question.
 ** After a type conversion, the result of calling sqlite3_column_type()
 ** is undefined, though harmless.  Future
 ** versions of SQLite may change the behavior of sqlite3_column_type()
@@ -4729,7 +5158,7 @@ SQLITE_API int sqlite3_data_count(sqlite3_stmt *pStmt);
 ** the number of bytes in that string.
 ** ^If the result is NULL, then sqlite3_column_bytes16() returns zero.
 **
-** ^The values returned by [sqlite3_column_bytes()] and 
+** ^The values returned by [sqlite3_column_bytes()] and
 ** [sqlite3_column_bytes16()] do not include the zero terminators at the end
 ** of the string.  ^For clarity: the values returned by
 ** [sqlite3_column_bytes()] and [sqlite3_column_bytes16()] are the number of
@@ -4739,6 +5168,10 @@ SQLITE_API int sqlite3_data_count(sqlite3_stmt *pStmt);
 ** even empty strings, are always zero-terminated.  ^The return
 ** value from sqlite3_column_blob() for a zero-length BLOB is a NULL pointer.
 **
+** ^Strings returned by sqlite3_column_text16() always have the endianness
+** which is native to the platform, regardless of the text encoding set
+** for the database.
+**
 ** <b>Warning:</b> ^The object returned by [sqlite3_column_value()] is an
 ** [unprotected sqlite3_value] object.  In a multithreaded environment,
 ** an unprotected sqlite3_value object may only be used safely with
@@ -4748,11 +5181,11 @@ SQLITE_API int sqlite3_data_count(sqlite3_stmt *pStmt);
 ** to routines like [sqlite3_value_int()], [sqlite3_value_text()],
 ** or [sqlite3_value_bytes()], the behavior is not threadsafe.
 ** Hence, the sqlite3_column_value() interface
-** is normally only useful within the implementation of 
+** is normally only useful within the implementation of
 ** [application-defined SQL functions] or [virtual tables], not within
 ** top-level application code.
 **
-** The these routines may attempt to convert the datatype of the result.
+** These routines may attempt to convert the datatype of the result.
 ** ^For example, if the internal representation is FLOAT and a text result
 ** is requested, [sqlite3_snprintf()] is used internally to perform the
 ** conversion automatically.  ^(The following table details the conversions
@@ -4777,7 +5210,7 @@ SQLITE_API int sqlite3_data_count(sqlite3_stmt *pStmt);
 ** <tr><td>  TEXT    <td>   BLOB    <td> No change
 ** <tr><td>  BLOB    <td> INTEGER   <td> [CAST] to INTEGER
 ** <tr><td>  BLOB    <td>  FLOAT    <td> [CAST] to REAL
-** <tr><td>  BLOB    <td>   TEXT    <td> Add a zero terminator if needed
+** <tr><td>  BLOB    <td>   TEXT    <td> [CAST] to TEXT, ensure zero terminator
 ** </table>
 ** </blockquote>)^
 **
@@ -4901,20 +5334,33 @@ SQLITE_API int sqlite3_finalize(sqlite3_stmt *pStmt);
 ** ^The [sqlite3_reset(S)] interface resets the [prepared statement] S
 ** back to the beginning of its program.
 **
-** ^If the most recent call to [sqlite3_step(S)] for the
-** [prepared statement] S returned [SQLITE_ROW] or [SQLITE_DONE],
-** or if [sqlite3_step(S)] has never before been called on S,
-** then [sqlite3_reset(S)] returns [SQLITE_OK].
+** ^The return code from [sqlite3_reset(S)] indicates whether or not
+** the previous evaluation of prepared statement S completed successfully.
+** ^If [sqlite3_step(S)] has never before been called on S or if
+** [sqlite3_step(S)] has not been called since the previous call
+** to [sqlite3_reset(S)], then [sqlite3_reset(S)] will return
+** [SQLITE_OK].
 **
 ** ^If the most recent call to [sqlite3_step(S)] for the
 ** [prepared statement] S indicated an error, then
 ** [sqlite3_reset(S)] returns an appropriate [error code].
+** ^The [sqlite3_reset(S)] interface might also return an [error code]
+** if there were no prior errors but the process of resetting
+** the prepared statement caused a new error. ^For example, if an
+** [INSERT] statement with a [RETURNING] clause is only stepped one time,
+** that one call to [sqlite3_step(S)] might return SQLITE_ROW but
+** the overall statement might still fail and the [sqlite3_reset(S)] call
+** might return SQLITE_BUSY if locking constraints prevent the
+** database change from committing.  Therefore, it is important that
+** applications check the return code from [sqlite3_reset(S)] even if
+** no prior call to [sqlite3_step(S)] indicated a problem.
 **
 ** ^The [sqlite3_reset(S)] interface does not change the values
 ** of any [sqlite3_bind_blob|bindings] on the [prepared statement] S.
 */
 SQLITE_API int sqlite3_reset(sqlite3_stmt *pStmt);
 
+
 /*
 ** CAPI3REF: Create Or Redefine SQL Functions
 ** KEYWORDS: {function creation routines}
@@ -4923,8 +5369,8 @@ SQLITE_API int sqlite3_reset(sqlite3_stmt *pStmt);
 ** ^These functions (collectively known as "function creation routines")
 ** are used to add SQL functions or aggregates or to redefine the behavior
 ** of existing SQL functions or aggregates. The only differences between
-** the three "sqlite3_create_function*" routines are the text encoding 
-** expected for the second parameter (the name of the function being 
+** the three "sqlite3_create_function*" routines are the text encoding
+** expected for the second parameter (the name of the function being
 ** created) and the presence or absence of a destructor callback for
 ** the application data pointer. Function sqlite3_create_window_function()
 ** is similar, but allows the user to supply the extra callback functions
@@ -4938,7 +5384,7 @@ SQLITE_API int sqlite3_reset(sqlite3_stmt *pStmt);
 ** ^The second parameter is the name of the SQL function to be created or
 ** redefined.  ^The length of the name is limited to 255 bytes in a UTF-8
 ** representation, exclusive of the zero-terminator.  ^Note that the name
-** length limit is in UTF-8 bytes, not characters nor UTF-16 bytes.  
+** length limit is in UTF-8 bytes, not characters nor UTF-16 bytes.
 ** ^Any attempt to create a function with a longer name
 ** will result in [SQLITE_MISUSE] being returned.
 **
@@ -4953,7 +5399,7 @@ SQLITE_API int sqlite3_reset(sqlite3_stmt *pStmt);
 ** ^The fourth parameter, eTextRep, specifies what
 ** [SQLITE_UTF8 | text encoding] this SQL function prefers for
 ** its parameters.  The application should set this parameter to
-** [SQLITE_UTF16LE] if the function implementation invokes 
+** [SQLITE_UTF16LE] if the function implementation invokes
 ** [sqlite3_value_text16le()] on an input, or [SQLITE_UTF16BE] if the
 ** implementation invokes [sqlite3_value_text16be()] on an input, or
 ** [SQLITE_UTF16] if [sqlite3_value_text16()] is used, or [SQLITE_UTF8]
@@ -4976,17 +5422,15 @@ SQLITE_API int sqlite3_reset(sqlite3_stmt *pStmt);
 ** within VIEWs, TRIGGERs, CHECK constraints, generated column expressions,
 ** index expressions, or the WHERE clause of partial indexes.
 **
-** <span style="background-color:#ffff90;">
 ** For best security, the [SQLITE_DIRECTONLY] flag is recommended for
 ** all application-defined SQL functions that do not need to be
 ** used inside of triggers, view, CHECK constraints, or other elements of
-** the database schema.  This flags is especially recommended for SQL 
+** the database schema.  This flags is especially recommended for SQL
 ** functions that have side effects or reveal internal application state.
 ** Without this flag, an attacker might be able to modify the schema of
 ** a database file to include invocations of the function with parameters
 ** chosen by the attacker, which the application will then execute when
 ** the database file is opened and read.
-** </span>
 **
 ** ^(The fifth parameter is an arbitrary pointer.  The implementation of the
 ** function can gain access to this pointer using [sqlite3_user_data()].)^
@@ -5001,21 +5445,21 @@ SQLITE_API int sqlite3_reset(sqlite3_stmt *pStmt);
 ** SQL function or aggregate, pass NULL pointers for all three function
 ** callbacks.
 **
-** ^The sixth, seventh, eighth and ninth parameters (xStep, xFinal, xValue 
+** ^The sixth, seventh, eighth and ninth parameters (xStep, xFinal, xValue
 ** and xInverse) passed to sqlite3_create_window_function are pointers to
 ** C-language callbacks that implement the new function. xStep and xFinal
 ** must both be non-NULL. xValue and xInverse may either both be NULL, in
-** which case a regular aggregate function is created, or must both be 
+** which case a regular aggregate function is created, or must both be
 ** non-NULL, in which case the new function may be used as either an aggregate
 ** or aggregate window function. More details regarding the implementation
-** of aggregate window functions are 
+** of aggregate window functions are
 ** [user-defined window functions|available here].
 **
 ** ^(If the final parameter to sqlite3_create_function_v2() or
 ** sqlite3_create_window_function() is not NULL, then it is destructor for
-** the application data pointer. The destructor is invoked when the function 
-** is deleted, either by being overloaded or when the database connection 
-** closes.)^ ^The destructor is also invoked if the call to 
+** the application data pointer. The destructor is invoked when the function
+** is deleted, either by being overloaded or when the database connection
+** closes.)^ ^The destructor is also invoked if the call to
 ** sqlite3_create_function_v2() fails.  ^When the destructor callback is
 ** invoked, it is passed a single argument which is a copy of the application
 ** data pointer which was the fifth parameter to sqlite3_create_function_v2().
@@ -5028,7 +5472,7 @@ SQLITE_API int sqlite3_reset(sqlite3_stmt *pStmt);
 ** nArg parameter is a better match than a function implementation with
 ** a negative nArg.  ^A function where the preferred text encoding
 ** matches the database encoding is a better
-** match than a function where the encoding is different.  
+** match than a function where the encoding is different.
 ** ^A function where the encoding difference is between UTF16le and UTF16be
 ** is a closer match than a function where the encoding difference is
 ** between UTF8 and UTF16.
@@ -5100,7 +5544,7 @@ SQLITE_API int sqlite3_create_window_function(
 /*
 ** CAPI3REF: Function Flags
 **
-** These constants may be ORed together with the 
+** These constants may be ORed together with the
 ** [SQLITE_UTF8 | preferred text encoding] as the fourth argument
 ** to [sqlite3_create_function()], [sqlite3_create_function16()], or
 ** [sqlite3_create_function_v2()].
@@ -5116,16 +5560,27 @@ SQLITE_API int sqlite3_create_window_function(
 ** SQLite might also optimize deterministic functions by factoring them
 ** out of inner loops.
 ** </dd>
-** 
+**
 ** [[SQLITE_DIRECTONLY]] <dt>SQLITE_DIRECTONLY</dt><dd>
 ** The SQLITE_DIRECTONLY flag means that the function may only be invoked
-** from top-level SQL, and cannot be used in VIEWs or TRIGGERs nor in 
+** from top-level SQL, and cannot be used in VIEWs or TRIGGERs nor in
 ** schema structures such as [CHECK constraints], [DEFAULT clauses],
 ** [expression indexes], [partial indexes], or [generated columns].
-** The SQLITE_DIRECTONLY flags is a security feature which is recommended
-** for all [application-defined SQL functions], and especially for functions
-** that have side-effects or that could potentially leak sensitive
-** information.
+** <p>
+** The SQLITE_DIRECTONLY flag is recommended for any
+** [application-defined SQL function]
+** that has side-effects or that could potentially leak sensitive information.
+** This will prevent attacks in which an application is tricked
+** into using a database file that has had its schema surreptitiously
+** modified to invoke the application-defined function in ways that are
+** harmful.
+** <p>
+** Some people say it is good practice to set SQLITE_DIRECTONLY on all
+** [application-defined SQL functions], regardless of whether or not they
+** are security sensitive, as doing so prevents those functions from being used
+** inside of the database schema, and thus ensures that the database
+** can be inspected and modified using generic tools (such as the [CLI])
+** that do not have access to the application-defined functions.
 ** </dd>
 **
 ** [[SQLITE_INNOCUOUS]] <dt>SQLITE_INNOCUOUS</dt><dd>
@@ -5152,13 +5607,36 @@ SQLITE_API int sqlite3_create_window_function(
 ** </dd>
 **
 ** [[SQLITE_SUBTYPE]] <dt>SQLITE_SUBTYPE</dt><dd>
-** The SQLITE_SUBTYPE flag indicates to SQLite that a function may call
+** The SQLITE_SUBTYPE flag indicates to SQLite that a function might call
 ** [sqlite3_value_subtype()] to inspect the sub-types of its arguments.
-** Specifying this flag makes no difference for scalar or aggregate user
-** functions. However, if it is not specified for a user-defined window
-** function, then any sub-types belonging to arguments passed to the window
-** function may be discarded before the window function is called (i.e.
-** sqlite3_value_subtype() will always return 0).
+** This flag instructs SQLite to omit some corner-case optimizations that
+** might disrupt the operation of the [sqlite3_value_subtype()] function,
+** causing it to return zero rather than the correct subtype().
+** All SQL functions that invoke [sqlite3_value_subtype()] should have this
+** property.  If the SQLITE_SUBTYPE property is omitted, then the return
+** value from [sqlite3_value_subtype()] might sometimes be zero even though
+** a non-zero subtype was specified by the function argument expression.
+**
+** [[SQLITE_RESULT_SUBTYPE]] <dt>SQLITE_RESULT_SUBTYPE</dt><dd>
+** The SQLITE_RESULT_SUBTYPE flag indicates to SQLite that a function might call
+** [sqlite3_result_subtype()] to cause a sub-type to be associated with its
+** result.
+** Every function that invokes [sqlite3_result_subtype()] should have this
+** property.  If it does not, then the call to [sqlite3_result_subtype()]
+** might become a no-op if the function is used as term in an
+** [expression index].  On the other hand, SQL functions that never invoke
+** [sqlite3_result_subtype()] should avoid setting this property, as the
+** purpose of this property is to disable certain optimizations that are
+** incompatible with subtypes.
+**
+** [[SQLITE_SELFORDER1]] <dt>SQLITE_SELFORDER1</dt><dd>
+** The SQLITE_SELFORDER1 flag indicates that the function is an aggregate
+** that internally orders the values provided to the first argument.  The
+** ordered-set aggregate SQL notation with a single ORDER BY term can be
+** used to invoke this function.  If the ordered-set aggregate notation is
+** used on a function that lacks this flag, then an error is raised. Note
+** that the ordered-set aggregate syntax is only available if SQLite is
+** built using the -DSQLITE_ENABLE_ORDERED_SET_AGGREGATES compile-time option.
 ** </dd>
 ** </dl>
 */
@@ -5166,13 +5644,15 @@ SQLITE_API int sqlite3_create_window_function(
 #define SQLITE_DIRECTONLY       0x000080000
 #define SQLITE_SUBTYPE          0x000100000
 #define SQLITE_INNOCUOUS        0x000200000
+#define SQLITE_RESULT_SUBTYPE   0x001000000
+#define SQLITE_SELFORDER1       0x002000000
 
 /*
 ** CAPI3REF: Deprecated Functions
 ** DEPRECATED
 **
 ** These functions are [deprecated].  In order to maintain
-** backwards compatibility with older code, these functions continue 
+** backwards compatibility with older code, these functions continue
 ** to be supported.  However, new applications should avoid
 ** the use of these functions.  To encourage programmers to avoid
 ** these functions, we will not explain what they do.
@@ -5240,11 +5720,11 @@ SQLITE_API SQLITE_DEPRECATED int sqlite3_memory_alarm(void(*)(void*,sqlite3_int6
 ** sqlite3_value_text16be() and sqlite3_value_text16le() interfaces
 ** extract UTF-16 strings as big-endian and little-endian respectively.
 **
-** ^If [sqlite3_value] object V was initialized 
+** ^If [sqlite3_value] object V was initialized
 ** using [sqlite3_bind_pointer(S,I,P,X,D)] or [sqlite3_result_pointer(C,P,X,D)]
 ** and if X and Y are strings that compare equal according to strcmp(X,Y),
 ** then sqlite3_value_pointer(V,Y) will return the pointer P.  ^Otherwise,
-** sqlite3_value_pointer(V,Y) returns a NULL. The sqlite3_bind_pointer() 
+** sqlite3_value_pointer(V,Y) returns a NULL. The sqlite3_bind_pointer()
 ** routine is part of the [pointer passing interface] added for SQLite 3.20.0.
 **
 ** ^(The sqlite3_value_type(V) interface returns the
@@ -5331,6 +5811,28 @@ SQLITE_API int sqlite3_value_numeric_type(sqlite3_value*);
 SQLITE_API int sqlite3_value_nochange(sqlite3_value*);
 SQLITE_API int sqlite3_value_frombind(sqlite3_value*);
 
+/*
+** CAPI3REF: Report the internal text encoding state of an sqlite3_value object
+** METHOD: sqlite3_value
+**
+** ^(The sqlite3_value_encoding(X) interface returns one of [SQLITE_UTF8],
+** [SQLITE_UTF16BE], or [SQLITE_UTF16LE] according to the current text encoding
+** of the value X, assuming that X has type TEXT.)^  If sqlite3_value_type(X)
+** returns something other than SQLITE_TEXT, then the return value from
+** sqlite3_value_encoding(X) is meaningless.  ^Calls to
+** [sqlite3_value_text(X)], [sqlite3_value_text16(X)], [sqlite3_value_text16be(X)],
+** [sqlite3_value_text16le(X)], [sqlite3_value_bytes(X)], or
+** [sqlite3_value_bytes16(X)] might change the encoding of the value X and
+** thus change the return from subsequent calls to sqlite3_value_encoding(X).
+**
+** This routine is intended for used by applications that test and validate
+** the SQLite implementation.  This routine is inquiring about the opaque
+** internal state of an [sqlite3_value] object.  Ordinary applications should
+** not need to know what the internal state of an sqlite3_value object is and
+** hence should not need to use this interface.
+*/
+SQLITE_API int sqlite3_value_encoding(sqlite3_value*);
+
 /*
 ** CAPI3REF: Finding The Subtype Of SQL Values
 ** METHOD: sqlite3_value
@@ -5340,6 +5842,12 @@ SQLITE_API int sqlite3_value_frombind(sqlite3_value*);
 ** information can be used to pass a limited amount of context from
 ** one SQL function to another.  Use the [sqlite3_result_subtype()]
 ** routine to set the subtype for the return value of an SQL function.
+**
+** Every [application-defined SQL function] that invokes this interface
+** should include the [SQLITE_SUBTYPE] property in the text
+** encoding argument when the function is [sqlite3_create_function|registered].
+** If the [SQLITE_SUBTYPE] property is omitted, then sqlite3_value_subtype()
+** might return zero instead of the upstream subtype in some corner cases.
 */
 SQLITE_API unsigned int sqlite3_value_subtype(sqlite3_value*);
 
@@ -5351,7 +5859,8 @@ SQLITE_API unsigned int sqlite3_value_subtype(sqlite3_value*);
 ** object D and returns a pointer to that copy.  ^The [sqlite3_value] returned
 ** is a [protected sqlite3_value] object even if the input is not.
 ** ^The sqlite3_value_dup(V) interface returns NULL if V is NULL or if a
-** memory allocation fails.
+** memory allocation fails. ^If V is a [pointer value], then the result
+** of sqlite3_value_dup(V) is a NULL value.
 **
 ** ^The sqlite3_value_free(V) interface frees an [sqlite3_value] object
 ** previously obtained from [sqlite3_value_dup()].  ^If V is a NULL pointer
@@ -5367,7 +5876,7 @@ SQLITE_API void sqlite3_value_free(sqlite3_value*);
 ** Implementations of aggregate SQL functions use this
 ** routine to allocate memory for storing their state.
 **
-** ^The first time the sqlite3_aggregate_context(C,N) routine is called 
+** ^The first time the sqlite3_aggregate_context(C,N) routine is called
 ** for a particular aggregate function, SQLite allocates
 ** N bytes of memory, zeroes out that memory, and returns a pointer
 ** to the new memory. ^On second and subsequent calls to
@@ -5380,19 +5889,19 @@ SQLITE_API void sqlite3_value_free(sqlite3_value*);
 ** In those cases, sqlite3_aggregate_context() might be called for the
 ** first time from within xFinal().)^
 **
-** ^The sqlite3_aggregate_context(C,N) routine returns a NULL pointer 
+** ^The sqlite3_aggregate_context(C,N) routine returns a NULL pointer
 ** when first called if N is less than or equal to zero or if a memory
-** allocate error occurs.
+** allocation error occurs.
 **
 ** ^(The amount of space allocated by sqlite3_aggregate_context(C,N) is
 ** determined by the N parameter on first successful call.  Changing the
-** value of N in any subsequents call to sqlite3_aggregate_context() within
+** value of N in any subsequent call to sqlite3_aggregate_context() within
 ** the same aggregate function instance will not resize the memory
 ** allocation.)^  Within the xFinal callback, it is customary to set
-** N=0 in calls to sqlite3_aggregate_context(C,N) so that no 
+** N=0 in calls to sqlite3_aggregate_context(C,N) so that no
 ** pointless memory allocations occur.
 **
-** ^SQLite automatically frees the memory allocated by 
+** ^SQLite automatically frees the memory allocated by
 ** sqlite3_aggregate_context() when the aggregate query concludes.
 **
 ** The first parameter must be a copy of the
@@ -5437,48 +5946,56 @@ SQLITE_API sqlite3 *sqlite3_context_db_handle(sqlite3_context*);
 ** METHOD: sqlite3_context
 **
 ** These functions may be used by (non-aggregate) SQL functions to
-** associate metadata with argument values. If the same value is passed to
-** multiple invocations of the same SQL function during query execution, under
-** some circumstances the associated metadata may be preserved.  An example
-** of where this might be useful is in a regular-expression matching
-** function. The compiled version of the regular expression can be stored as
-** metadata associated with the pattern string.  
+** associate auxiliary data with argument values. If the same argument
+** value is passed to multiple invocations of the same SQL function during
+** query execution, under some circumstances the associated auxiliary data
+** might be preserved.  An example of where this might be useful is in a
+** regular-expression matching function. The compiled version of the regular
+** expression can be stored as auxiliary data associated with the pattern string.
 ** Then as long as the pattern string remains the same,
 ** the compiled regular expression can be reused on multiple
 ** invocations of the same function.
 **
-** ^The sqlite3_get_auxdata(C,N) interface returns a pointer to the metadata
+** ^The sqlite3_get_auxdata(C,N) interface returns a pointer to the auxiliary data
 ** associated by the sqlite3_set_auxdata(C,N,P,X) function with the Nth argument
 ** value to the application-defined function.  ^N is zero for the left-most
-** function argument.  ^If there is no metadata
+** function argument.  ^If there is no auxiliary data
 ** associated with the function argument, the sqlite3_get_auxdata(C,N) interface
 ** returns a NULL pointer.
 **
-** ^The sqlite3_set_auxdata(C,N,P,X) interface saves P as metadata for the N-th
-** argument of the application-defined function.  ^Subsequent
+** ^The sqlite3_set_auxdata(C,N,P,X) interface saves P as auxiliary data for the
+** N-th argument of the application-defined function.  ^Subsequent
 ** calls to sqlite3_get_auxdata(C,N) return P from the most recent
-** sqlite3_set_auxdata(C,N,P,X) call if the metadata is still valid or
-** NULL if the metadata has been discarded.
+** sqlite3_set_auxdata(C,N,P,X) call if the auxiliary data is still valid or
+** NULL if the auxiliary data has been discarded.
 ** ^After each call to sqlite3_set_auxdata(C,N,P,X) where X is not NULL,
 ** SQLite will invoke the destructor function X with parameter P exactly
-** once, when the metadata is discarded.
-** SQLite is free to discard the metadata at any time, including: <ul>
+** once, when the auxiliary data is discarded.
+** SQLite is free to discard the auxiliary data at any time, including: <ul>
 ** <li> ^(when the corresponding function parameter changes)^, or
 ** <li> ^(when [sqlite3_reset()] or [sqlite3_finalize()] is called for the
 **      SQL statement)^, or
 ** <li> ^(when sqlite3_set_auxdata() is invoked again on the same
 **       parameter)^, or
-** <li> ^(during the original sqlite3_set_auxdata() call when a memory 
-**      allocation error occurs.)^ </ul>
+** <li> ^(during the original sqlite3_set_auxdata() call when a memory
+**      allocation error occurs.)^
+** <li> ^(during the original sqlite3_set_auxdata() call if the function
+**      is evaluated during query planning instead of during query execution,
+**      as sometimes happens with [SQLITE_ENABLE_STAT4].)^ </ul>
 **
-** Note the last bullet in particular.  The destructor X in 
+** Note the last two bullets in particular.  The destructor X in
 ** sqlite3_set_auxdata(C,N,P,X) might be called immediately, before the
 ** sqlite3_set_auxdata() interface even returns.  Hence sqlite3_set_auxdata()
 ** should be called near the end of the function implementation and the
 ** function implementation should not make any use of P after
-** sqlite3_set_auxdata() has been called.
-**
-** ^(In practice, metadata is preserved between function calls for
+** sqlite3_set_auxdata() has been called.  Furthermore, a call to
+** sqlite3_get_auxdata() that occurs immediately after a corresponding call
+** to sqlite3_set_auxdata() might still return NULL if an out-of-memory
+** condition occurred during the sqlite3_set_auxdata() call or if the
+** function is being evaluated during query planning rather than during
+** query execution.
+**
+** ^(In practice, auxiliary data is preserved between function calls for
 ** function parameters that are compile-time constants, including literal
 ** values and [parameters] and expressions composed from the same.)^
 **
@@ -5488,10 +6005,67 @@ SQLITE_API sqlite3 *sqlite3_context_db_handle(sqlite3_context*);
 **
 ** These routines must be called from the same thread in which
 ** the SQL function is running.
+**
+** See also: [sqlite3_get_clientdata()] and [sqlite3_set_clientdata()].
 */
 SQLITE_API void *sqlite3_get_auxdata(sqlite3_context*, int N);
 SQLITE_API void sqlite3_set_auxdata(sqlite3_context*, int N, void*, void (*)(void*));
 
+/*
+** CAPI3REF: Database Connection Client Data
+** METHOD: sqlite3
+**
+** These functions are used to associate one or more named pointers
+** with a [database connection].
+** A call to sqlite3_set_clientdata(D,N,P,X) causes the pointer P
+** to be attached to [database connection] D using name N.  Subsequent
+** calls to sqlite3_get_clientdata(D,N) will return a copy of pointer P
+** or a NULL pointer if there were no prior calls to
+** sqlite3_set_clientdata() with the same values of D and N.
+** Names are compared using strcmp() and are thus case sensitive.
+**
+** If P and X are both non-NULL, then the destructor X is invoked with
+** argument P on the first of the following occurrences:
+** <ul>
+** <li> An out-of-memory error occurs during the call to
+**      sqlite3_set_clientdata() which attempts to register pointer P.
+** <li> A subsequent call to sqlite3_set_clientdata(D,N,P,X) is made
+**      with the same D and N parameters.
+** <li> The database connection closes.  SQLite does not make any guarantees
+**      about the order in which destructors are called, only that all
+**      destructors will be called exactly once at some point during the
+**      database connection closing process.
+** </ul>
+**
+** SQLite does not do anything with client data other than invoke
+** destructors on the client data at the appropriate time.  The intended
+** use for client data is to provide a mechanism for wrapper libraries
+** to store additional information about an SQLite database connection.
+**
+** There is no limit (other than available memory) on the number of different
+** client data pointers (with different names) that can be attached to a
+** single database connection.  However, the implementation is optimized
+** for the case of having only one or two different client data names.
+** Applications and wrapper libraries are discouraged from using more than
+** one client data name each.
+**
+** There is no way to enumerate the client data pointers
+** associated with a database connection.  The N parameter can be thought
+** of as a secret key such that only code that knows the secret key is able
+** to access the associated data.
+**
+** Security Warning:  These interfaces should not be exposed in scripting
+** languages or in other circumstances where it might be possible for an
+** an attacker to invoke them.  Any agent that can invoke these interfaces
+** can probably also take control of the process.
+**
+** Database connection client data is only available for SQLite
+** version 3.44.0 ([dateof:3.44.0]) and later.
+**
+** See also: [sqlite3_set_auxdata()] and [sqlite3_get_auxdata()].
+*/
+SQLITE_API void *sqlite3_get_clientdata(sqlite3*,const char*);
+SQLITE_API int sqlite3_set_clientdata(sqlite3*, const char*, void*, void(*)(void*));
 
 /*
 ** CAPI3REF: Constants Defining Special Destructor Behavior
@@ -5543,8 +6117,9 @@ typedef void (*sqlite3_destructor_type)(void*);
 ** 2nd parameter of sqlite3_result_error() or sqlite3_result_error16()
 ** as the text of an error message.  ^SQLite interprets the error
 ** message string from sqlite3_result_error() as UTF-8. ^SQLite
-** interprets the string from sqlite3_result_error16() as UTF-16 in native
-** byte order.  ^If the third parameter to sqlite3_result_error()
+** interprets the string from sqlite3_result_error16() as UTF-16 using
+** the same [byte-order determination rules] as [sqlite3_bind_text16()].
+** ^If the third parameter to sqlite3_result_error()
 ** or sqlite3_result_error16() is negative then SQLite takes as the error
 ** message all text up through the first zero character.
 ** ^If the third parameter to sqlite3_result_error() or
@@ -5586,9 +6161,10 @@ typedef void (*sqlite3_destructor_type)(void*);
 ** of [SQLITE_UTF8], [SQLITE_UTF16], [SQLITE_UTF16BE], or [SQLITE_UTF16LE].
 ** ^SQLite takes the text result from the application from
 ** the 2nd parameter of the sqlite3_result_text* interfaces.
-** ^If the 3rd parameter to the sqlite3_result_text* interfaces
-** is negative, then SQLite takes result text from the 2nd parameter
-** through the first zero character.
+** ^If the 3rd parameter to any of the sqlite3_result_text* interfaces
+** other than sqlite3_result_text64() is negative, then SQLite computes
+** the string length itself by searching the 2nd parameter for the first
+** zero character.
 ** ^If the 3rd parameter to the sqlite3_result_text* interfaces
 ** is non-negative, then as many bytes (not characters) of the text
 ** pointed to by the 2nd parameter are taken as the application-defined
@@ -5612,6 +6188,25 @@ typedef void (*sqlite3_destructor_type)(void*);
 ** then SQLite makes a copy of the result into space obtained
 ** from [sqlite3_malloc()] before it returns.
 **
+** ^For the sqlite3_result_text16(), sqlite3_result_text16le(), and
+** sqlite3_result_text16be() routines, and for sqlite3_result_text64()
+** when the encoding is not UTF8, if the input UTF16 begins with a
+** byte-order mark (BOM, U+FEFF) then the BOM is removed from the
+** string and the rest of the string is interpreted according to the
+** byte-order specified by the BOM.  ^The byte-order specified by
+** the BOM at the beginning of the text overrides the byte-order
+** specified by the interface procedure.  ^So, for example, if
+** sqlite3_result_text16le() is invoked with text that begins
+** with bytes 0xfe, 0xff (a big-endian byte-order mark) then the
+** first two bytes of input are skipped and the remaining input
+** is interpreted as UTF16BE text.
+**
+** ^For UTF16 input text to the sqlite3_result_text16(),
+** sqlite3_result_text16be(), sqlite3_result_text16le(), and
+** sqlite3_result_text64() routines, if the text contains invalid
+** UTF16 characters, the invalid characters might be converted
+** into the unicode replacement character, U+FFFD.
+**
 ** ^The sqlite3_result_value() interface sets the result of
 ** the application-defined function to be a copy of the
 ** [unprotected sqlite3_value] object specified by the 2nd parameter.  ^The
@@ -5624,7 +6219,7 @@ typedef void (*sqlite3_destructor_type)(void*);
 **
 ** ^The sqlite3_result_pointer(C,P,T,D) interface sets the result to an
 ** SQL NULL value, just like [sqlite3_result_null(C)], except that it
-** also associates the host-language pointer P or type T with that 
+** also associates the host-language pointer P or type T with that
 ** NULL value such that the pointer can be retrieved within an
 ** [application-defined SQL function] using [sqlite3_value_pointer()].
 ** ^If the D parameter is not NULL, then it is a pointer to a destructor
@@ -5666,12 +6261,26 @@ SQLITE_API int sqlite3_result_zeroblob64(sqlite3_context*, sqlite3_uint64 n);
 ** METHOD: sqlite3_context
 **
 ** The sqlite3_result_subtype(C,T) function causes the subtype of
-** the result from the [application-defined SQL function] with 
-** [sqlite3_context] C to be the value T.  Only the lower 8 bits 
+** the result from the [application-defined SQL function] with
+** [sqlite3_context] C to be the value T.  Only the lower 8 bits
 ** of the subtype T are preserved in current versions of SQLite;
 ** higher order bits are discarded.
 ** The number of subtype bytes preserved by SQLite might increase
 ** in future releases of SQLite.
+**
+** Every [application-defined SQL function] that invokes this interface
+** should include the [SQLITE_RESULT_SUBTYPE] property in its
+** text encoding argument when the SQL function is
+** [sqlite3_create_function|registered].  If the [SQLITE_RESULT_SUBTYPE]
+** property is omitted from the function that invokes sqlite3_result_subtype(),
+** then in some cases the sqlite3_result_subtype() might fail to set
+** the result subtype.
+**
+** If SQLite is compiled with -DSQLITE_STRICT_SUBTYPE=1, then any
+** SQL function that invokes the sqlite3_result_subtype() interface
+** and that does not have the SQLITE_RESULT_SUBTYPE property will raise
+** an error.  Future versions of SQLite might enable -DSQLITE_STRICT_SUBTYPE=1
+** by default.
 */
 SQLITE_API void sqlite3_result_subtype(sqlite3_context*,unsigned int);
 
@@ -5714,7 +6323,7 @@ SQLITE_API void sqlite3_result_subtype(sqlite3_context*,unsigned int);
 ** deleted.  ^When all collating functions having the same name are deleted,
 ** that collation is no longer usable.
 **
-** ^The collating function callback is invoked with a copy of the pArg 
+** ^The collating function callback is invoked with a copy of the pArg
 ** application data pointer and with two strings in the encoding specified
 ** by the eTextRep argument.  The two integer parameters to the collating
 ** function callback are the length of the two strings, in bytes. The collating
@@ -5745,36 +6354,36 @@ SQLITE_API void sqlite3_result_subtype(sqlite3_context*,unsigned int);
 ** calls to the collation creation functions or when the
 ** [database connection] is closed using [sqlite3_close()].
 **
-** ^The xDestroy callback is <u>not</u> called if the 
+** ^The xDestroy callback is <u>not</u> called if the
 ** sqlite3_create_collation_v2() function fails.  Applications that invoke
-** sqlite3_create_collation_v2() with a non-NULL xDestroy argument should 
+** sqlite3_create_collation_v2() with a non-NULL xDestroy argument should
 ** check the return code and dispose of the application data pointer
 ** themselves rather than expecting SQLite to deal with it for them.
-** This is different from every other SQLite interface.  The inconsistency 
-** is unfortunate but cannot be changed without breaking backwards 
+** This is different from every other SQLite interface.  The inconsistency
+** is unfortunate but cannot be changed without breaking backwards
 ** compatibility.
 **
 ** See also:  [sqlite3_collation_needed()] and [sqlite3_collation_needed16()].
 */
 SQLITE_API int sqlite3_create_collation(
-  sqlite3*, 
-  const char *zName, 
-  int eTextRep, 
+  sqlite3*,
+  const char *zName,
+  int eTextRep,
   void *pArg,
   int(*xCompare)(void*,int,const void*,int,const void*)
 );
 SQLITE_API int sqlite3_create_collation_v2(
-  sqlite3*, 
-  const char *zName, 
-  int eTextRep, 
+  sqlite3*,
+  const char *zName,
+  int eTextRep,
   void *pArg,
   int(*xCompare)(void*,int,const void*,int,const void*),
   void(*xDestroy)(void*)
 );
 SQLITE_API int sqlite3_create_collation16(
-  sqlite3*, 
+  sqlite3*,
   const void *zName,
-  int eTextRep, 
+  int eTextRep,
   void *pArg,
   int(*xCompare)(void*,int,const void*,int,const void*)
 );
@@ -5807,64 +6416,19 @@ SQLITE_API int sqlite3_create_collation16(
 ** [sqlite3_create_collation_v2()].
 */
 SQLITE_API int sqlite3_collation_needed(
-  sqlite3*, 
-  void*, 
+  sqlite3*,
+  void*,
   void(*)(void*,sqlite3*,int eTextRep,const char*)
 );
 SQLITE_API int sqlite3_collation_needed16(
-  sqlite3*, 
+  sqlite3*,
   void*,
   void(*)(void*,sqlite3*,int eTextRep,const void*)
 );
 
-#ifdef SQLITE_HAS_CODEC
-/*
-** Specify the key for an encrypted database.  This routine should be
-** called right after sqlite3_open().
-**
-** The code to implement this API is not available in the public release
-** of SQLite.
-*/
-SQLITE_API int sqlite3_key(
-  sqlite3 *db,                   /* Database to be rekeyed */
-  const void *pKey, int nKey     /* The key */
-);
-SQLITE_API int sqlite3_key_v2(
-  sqlite3 *db,                   /* Database to be rekeyed */
-  const char *zDbName,           /* Name of the database */
-  const void *pKey, int nKey     /* The key */
-);
-
-/*
-** Change the key on an open database.  If the current database is not
-** encrypted, this routine will encrypt it.  If pNew==0 or nNew==0, the
-** database is decrypted.
-**
-** The code to implement this API is not available in the public release
-** of SQLite.
-*/
-SQLITE_API int sqlite3_rekey(
-  sqlite3 *db,                   /* Database to be rekeyed */
-  const void *pKey, int nKey     /* The new key */
-);
-SQLITE_API int sqlite3_rekey_v2(
-  sqlite3 *db,                   /* Database to be rekeyed */
-  const char *zDbName,           /* Name of the database */
-  const void *pKey, int nKey     /* The new key */
-);
-
-/*
-** Specify the activation key for a SEE database.  Unless 
-** activated, none of the SEE routines will work.
-*/
-SQLITE_API void sqlite3_activate_see(
-  const char *zPassPhrase        /* Activation phrase */
-);
-#endif
-
 #ifdef SQLITE_ENABLE_CEROD
 /*
-** Specify the activation key for a CEROD database.  Unless 
+** Specify the activation key for a CEROD database.  Unless
 ** activated, none of the CEROD routines will work.
 */
 SQLITE_API void sqlite3_activate_cerod(
@@ -5888,6 +6452,13 @@ SQLITE_API void sqlite3_activate_cerod(
 ** of the default VFS is not implemented correctly, or not implemented at
 ** all, then the behavior of sqlite3_sleep() may deviate from the description
 ** in the previous paragraphs.
+**
+** If a negative argument is passed to sqlite3_sleep() the results vary by
+** VFS and operating system.  Some system treat a negative argument as an
+** instruction to sleep forever.  Others understand it to mean do not sleep
+** at all. ^In SQLite version 3.42.0 and later, a negative
+** argument passed into sqlite3_sleep() is changed to zero before it is relayed
+** down into the xSleep method of the VFS.
 */
 SQLITE_API int sqlite3_sleep(int);
 
@@ -5920,7 +6491,7 @@ SQLITE_API int sqlite3_sleep(int);
 ** ^The [temp_store_directory pragma] may modify this variable and cause
 ** it to point to memory obtained from [sqlite3_malloc].  ^Furthermore,
 ** the [temp_store_directory pragma] always assumes that any string
-** that this variable points to is held in memory obtained from 
+** that this variable points to is held in memory obtained from
 ** [sqlite3_malloc] and the pragma may attempt to free that memory
 ** using [sqlite3_free].
 ** Hence, if this variable is modified directly, either it should be
@@ -5977,7 +6548,7 @@ SQLITE_API SQLITE_EXTERN char *sqlite3_temp_directory;
 ** ^The [data_store_directory pragma] may modify this variable and cause
 ** it to point to memory obtained from [sqlite3_malloc].  ^Furthermore,
 ** the [data_store_directory pragma] always assumes that any string
-** that this variable points to is held in memory obtained from 
+** that this variable points to is held in memory obtained from
 ** [sqlite3_malloc] and the pragma may attempt to free that memory
 ** using [sqlite3_free].
 ** Hence, if this variable is modified directly, either it should be
@@ -6058,6 +6629,28 @@ SQLITE_API int sqlite3_get_autocommit(sqlite3*);
 */
 SQLITE_API sqlite3 *sqlite3_db_handle(sqlite3_stmt*);
 
+/*
+** CAPI3REF: Return The Schema Name For A Database Connection
+** METHOD: sqlite3
+**
+** ^The sqlite3_db_name(D,N) interface returns a pointer to the schema name
+** for the N-th database on database connection D, or a NULL pointer of N is
+** out of range.  An N value of 0 means the main database file.  An N of 1 is
+** the "temp" schema.  Larger values of N correspond to various ATTACH-ed
+** databases.
+**
+** Space to hold the string that is returned by sqlite3_db_name() is managed
+** by SQLite itself.  The string might be deallocated by any operation that
+** changes the schema, including [ATTACH] or [DETACH] or calls to
+** [sqlite3_serialize()] or [sqlite3_deserialize()], even operations that
+** occur on a different thread.  Applications that need to
+** remember the string long-term should make their own copy.  Applications that
+** are accessing the same database connection simultaneously on multiple
+** threads should mutex-protect calls to this API and should make their own
+** private copy of the result prior to releasing the mutex.
+*/
+SQLITE_API const char *sqlite3_db_name(sqlite3 *db, int N);
+
 /*
 ** CAPI3REF: Return The Filename For A Database Connection
 ** METHOD: sqlite3
@@ -6088,7 +6681,7 @@ SQLITE_API sqlite3 *sqlite3_db_handle(sqlite3_stmt*);
 ** <li> [sqlite3_filename_wal()]
 ** </ul>
 */
-SQLITE_API const char *sqlite3_db_filename(sqlite3 *db, const char *zDbName);
+SQLITE_API sqlite3_filename sqlite3_db_filename(sqlite3 *db, const char *zDbName);
 
 /*
 ** CAPI3REF: Determine if a database is read-only
@@ -6100,6 +6693,57 @@ SQLITE_API const char *sqlite3_db_filename(sqlite3 *db, const char *zDbName);
 */
 SQLITE_API int sqlite3_db_readonly(sqlite3 *db, const char *zDbName);
 
+/*
+** CAPI3REF: Determine the transaction state of a database
+** METHOD: sqlite3
+**
+** ^The sqlite3_txn_state(D,S) interface returns the current
+** [transaction state] of schema S in database connection D.  ^If S is NULL,
+** then the highest transaction state of any schema on database connection D
+** is returned.  Transaction states are (in order of lowest to highest):
+** <ol>
+** <li value="0"> SQLITE_TXN_NONE
+** <li value="1"> SQLITE_TXN_READ
+** <li value="2"> SQLITE_TXN_WRITE
+** </ol>
+** ^If the S argument to sqlite3_txn_state(D,S) is not the name of
+** a valid schema, then -1 is returned.
+*/
+SQLITE_API int sqlite3_txn_state(sqlite3*,const char *zSchema);
+
+/*
+** CAPI3REF: Allowed return values from sqlite3_txn_state()
+** KEYWORDS: {transaction state}
+**
+** These constants define the current transaction state of a database file.
+** ^The [sqlite3_txn_state(D,S)] interface returns one of these
+** constants in order to describe the transaction state of schema S
+** in [database connection] D.
+**
+** <dl>
+** [[SQLITE_TXN_NONE]] <dt>SQLITE_TXN_NONE</dt>
+** <dd>The SQLITE_TXN_NONE state means that no transaction is currently
+** pending.</dd>
+**
+** [[SQLITE_TXN_READ]] <dt>SQLITE_TXN_READ</dt>
+** <dd>The SQLITE_TXN_READ state means that the database is currently
+** in a read transaction.  Content has been read from the database file
+** but nothing in the database file has changed.  The transaction state
+** will advanced to SQLITE_TXN_WRITE if any changes occur and there are
+** no other conflicting concurrent write transactions.  The transaction
+** state will revert to SQLITE_TXN_NONE following a [ROLLBACK] or
+** [COMMIT].</dd>
+**
+** [[SQLITE_TXN_WRITE]] <dt>SQLITE_TXN_WRITE</dt>
+** <dd>The SQLITE_TXN_WRITE state means that the database is currently
+** in a write transaction.  Content has been written to the database file
+** but has not yet committed.  The transaction state will change to
+** to SQLITE_TXN_NONE at the next [ROLLBACK] or [COMMIT].</dd>
+*/
+#define SQLITE_TXN_NONE  0
+#define SQLITE_TXN_READ  1
+#define SQLITE_TXN_WRITE 2
+
 /*
 ** CAPI3REF: Find the next prepared statement
 ** METHOD: sqlite3
@@ -6166,6 +6810,72 @@ SQLITE_API sqlite3_stmt *sqlite3_next_stmt(sqlite3 *pDb, sqlite3_stmt *pStmt);
 SQLITE_API void *sqlite3_commit_hook(sqlite3*, int(*)(void*), void*);
 SQLITE_API void *sqlite3_rollback_hook(sqlite3*, void(*)(void *), void*);
 
+/*
+** CAPI3REF: Autovacuum Compaction Amount Callback
+** METHOD: sqlite3
+**
+** ^The sqlite3_autovacuum_pages(D,C,P,X) interface registers a callback
+** function C that is invoked prior to each autovacuum of the database
+** file.  ^The callback is passed a copy of the generic data pointer (P),
+** the schema-name of the attached database that is being autovacuumed,
+** the size of the database file in pages, the number of free pages,
+** and the number of bytes per page, respectively.  The callback should
+** return the number of free pages that should be removed by the
+** autovacuum.  ^If the callback returns zero, then no autovacuum happens.
+** ^If the value returned is greater than or equal to the number of
+** free pages, then a complete autovacuum happens.
+**
+** <p>^If there are multiple ATTACH-ed database files that are being
+** modified as part of a transaction commit, then the autovacuum pages
+** callback is invoked separately for each file.
+**
+** <p><b>The callback is not reentrant.</b> The callback function should
+** not attempt to invoke any other SQLite interface.  If it does, bad
+** things may happen, including segmentation faults and corrupt database
+** files.  The callback function should be a simple function that
+** does some arithmetic on its input parameters and returns a result.
+**
+** ^The X parameter to sqlite3_autovacuum_pages(D,C,P,X) is an optional
+** destructor for the P parameter.  ^If X is not NULL, then X(P) is
+** invoked whenever the database connection closes or when the callback
+** is overwritten by another invocation of sqlite3_autovacuum_pages().
+**
+** <p>^There is only one autovacuum pages callback per database connection.
+** ^Each call to the sqlite3_autovacuum_pages() interface overrides all
+** previous invocations for that database connection.  ^If the callback
+** argument (C) to sqlite3_autovacuum_pages(D,C,P,X) is a NULL pointer,
+** then the autovacuum steps callback is canceled.  The return value
+** from sqlite3_autovacuum_pages() is normally SQLITE_OK, but might
+** be some other error code if something goes wrong.  The current
+** implementation will only return SQLITE_OK or SQLITE_MISUSE, but other
+** return codes might be added in future releases.
+**
+** <p>If no autovacuum pages callback is specified (the usual case) or
+** a NULL pointer is provided for the callback,
+** then the default behavior is to vacuum all free pages.  So, in other
+** words, the default behavior is the same as if the callback function
+** were something like this:
+**
+** <blockquote><pre>
+** &nbsp;   unsigned int demonstration_autovac_pages_callback(
+** &nbsp;     void *pClientData,
+** &nbsp;     const char *zSchema,
+** &nbsp;     unsigned int nDbPage,
+** &nbsp;     unsigned int nFreePage,
+** &nbsp;     unsigned int nBytePerPage
+** &nbsp;   ){
+** &nbsp;     return nFreePage;
+** &nbsp;   }
+** </pre></blockquote>
+*/
+SQLITE_API int sqlite3_autovacuum_pages(
+  sqlite3 *db,
+  unsigned int(*)(void*,const char*,unsigned int,unsigned int,unsigned int),
+  void*,
+  void(*)(void*)
+);
+
+
 /*
 ** CAPI3REF: Data Change Notification Callbacks
 ** METHOD: sqlite3
@@ -6190,7 +6900,7 @@ SQLITE_API void *sqlite3_rollback_hook(sqlite3*, void(*)(void *), void*);
 ** ^In the case of an update, this is the [rowid] after the update takes place.
 **
 ** ^(The update hook is not invoked when internal system tables are
-** modified (i.e. sqlite_master and sqlite_sequence).)^
+** modified (i.e. sqlite_sequence).)^
 ** ^The update hook is not invoked when [WITHOUT ROWID] tables are modified.
 **
 ** ^In the current implementation, the update hook
@@ -6200,6 +6910,12 @@ SQLITE_API void *sqlite3_rollback_hook(sqlite3*, void(*)(void *), void*);
 ** The exceptions defined in this paragraph might change in a future
 ** release of SQLite.
 **
+** Whether the update hook is invoked before or after the
+** corresponding change is currently unspecified and may differ
+** depending on the type of change. Do not rely on the order of the
+** hook call with regards to the final result of the operation which
+** triggers the hook.
+**
 ** The update hook implementation must not do anything that will modify
 ** the database connection that invoked the update hook.  Any actions
 ** to modify the database connection must be deferred until after the
@@ -6216,7 +6932,7 @@ SQLITE_API void *sqlite3_rollback_hook(sqlite3*, void(*)(void *), void*);
 ** and [sqlite3_preupdate_hook()] interfaces.
 */
 SQLITE_API void *sqlite3_update_hook(
-  sqlite3*, 
+  sqlite3*,
   void(*)(void *,int ,char const *,char const *,sqlite3_int64),
   void*
 );
@@ -6229,8 +6945,13 @@ SQLITE_API void *sqlite3_update_hook(
 ** to the same database. Sharing is enabled if the argument is true
 ** and disabled if the argument is false.)^
 **
+** This interface is omitted if SQLite is compiled with
+** [-DSQLITE_OMIT_SHARED_CACHE].  The [-DSQLITE_OMIT_SHARED_CACHE]
+** compile-time option is recommended because the
+** [use of shared cache mode is discouraged].
+**
 ** ^Cache sharing is enabled and disabled for an entire process.
-** This is a change as of SQLite [version 3.5.0] ([dateof:3.5.0]). 
+** This is a change as of SQLite [version 3.5.0] ([dateof:3.5.0]).
 ** In prior versions of SQLite,
 ** sharing was enabled or disabled for each thread separately.
 **
@@ -6251,8 +6972,8 @@ SQLITE_API void *sqlite3_update_hook(
 ** with the [SQLITE_OPEN_SHAREDCACHE] flag.
 **
 ** Note: This method is disabled on MacOS X 10.7 and iOS version 5.0
-** and will always return SQLITE_MISUSE. On those systems, 
-** shared cache mode should be enabled per-database connection via 
+** and will always return SQLITE_MISUSE. On those systems,
+** shared cache mode should be enabled per-database connection via
 ** [sqlite3_open_v2()] with [SQLITE_OPEN_SHAREDCACHE].
 **
 ** This interface is threadsafe on processors where writing a
@@ -6305,7 +7026,7 @@ SQLITE_API int sqlite3_db_release_memory(sqlite3*);
 ** as heap memory usages approaches the limit.
 ** ^The soft heap limit is "soft" because even though SQLite strives to stay
 ** below the limit, it will exceed the limit rather than generate
-** an [SQLITE_NOMEM] error.  In other words, the soft heap limit 
+** an [SQLITE_NOMEM] error.  In other words, the soft heap limit
 ** is advisory only.
 **
 ** ^The sqlite3_hard_heap_limit64(N) interface sets a hard upper bound of
@@ -6327,7 +7048,7 @@ SQLITE_API int sqlite3_db_release_memory(sqlite3*);
 ** ^The soft heap limit may not be greater than the hard heap limit.
 ** ^If the hard heap limit is enabled and if sqlite3_soft_heap_limit(N)
 ** is invoked with a value of N that is greater than the hard heap limit,
-** the the soft heap limit is set to the value of the hard heap limit.
+** the soft heap limit is set to the value of the hard heap limit.
 ** ^The soft heap limit is automatically enabled whenever the hard heap
 ** limit is enabled. ^When sqlite3_hard_heap_limit64(N) is invoked and
 ** the soft heap limit is outside the range of 1..N, then the soft heap
@@ -6421,7 +7142,7 @@ SQLITE_API SQLITE_DEPRECATED void sqlite3_soft_heap_limit(int N);
 **
 ** ^If the specified table is actually a view, an [error code] is returned.
 **
-** ^If the specified column is "rowid", "oid" or "_rowid_" and the table 
+** ^If the specified column is "rowid", "oid" or "_rowid_" and the table
 ** is not a [WITHOUT ROWID] table and an
 ** [INTEGER PRIMARY KEY] column has been explicitly declared, then the output
 ** parameters are set for the explicitly declared column. ^(If there is no
@@ -6487,7 +7208,7 @@ SQLITE_API int sqlite3_table_column_metadata(
 ** prior to calling this API,
 ** otherwise an error will be returned.
 **
-** <b>Security warning:</b> It is recommended that the 
+** <b>Security warning:</b> It is recommended that the
 ** [SQLITE_DBCONFIG_ENABLE_LOAD_EXTENSION] method be used to enable only this
 ** interface.  The use of the [sqlite3_enable_load_extension()] interface
 ** should be avoided.  This will keep the SQL function [load_extension()]
@@ -6574,7 +7295,7 @@ SQLITE_API int sqlite3_auto_extension(void(*xEntryPoint)(void));
 ** ^The [sqlite3_cancel_auto_extension(X)] interface unregisters the
 ** initialization routine X that was registered using a prior call to
 ** [sqlite3_auto_extension(X)].  ^The [sqlite3_cancel_auto_extension(X)]
-** routine returns 1 if initialization routine X was successfully 
+** routine returns 1 if initialization routine X was successfully
 ** unregistered and it returns 0 if X was not on the list of initialization
 ** routines.
 */
@@ -6588,15 +7309,6 @@ SQLITE_API int sqlite3_cancel_auto_extension(void(*xEntryPoint)(void));
 */
 SQLITE_API void sqlite3_reset_auto_extension(void);
 
-/*
-** The interface to the virtual-table mechanism is currently considered
-** to be experimental.  The interface might change in incompatible ways.
-** If this is a problem for you, do not use the interface at this time.
-**
-** When the virtual-table mechanism stabilizes, we will declare the
-** interface fixed, support it indefinitely, and remove this comment.
-*/
-
 /*
 ** Structures used by the virtual table interface
 */
@@ -6609,8 +7321,8 @@ typedef struct sqlite3_module sqlite3_module;
 ** CAPI3REF: Virtual Table Object
 ** KEYWORDS: sqlite3_module {virtual table module}
 **
-** This structure, sometimes called a "virtual table module", 
-** defines the implementation of a [virtual table].  
+** This structure, sometimes called a "virtual table module",
+** defines the implementation of a [virtual table].
 ** This structure consists mostly of methods for the module.
 **
 ** ^A virtual table module is created by filling in a persistent
@@ -6649,7 +7361,7 @@ struct sqlite3_module {
                        void (**pxFunc)(sqlite3_context*,int,sqlite3_value**),
                        void **ppArg);
   int (*xRename)(sqlite3_vtab *pVtab, const char *zNew);
-  /* The methods above are in version 1 of the sqlite_module object. Those 
+  /* The methods above are in version 1 of the sqlite_module object. Those
   ** below are for version 2 and greater. */
   int (*xSavepoint)(sqlite3_vtab *pVTab, int);
   int (*xRelease)(sqlite3_vtab *pVTab, int);
@@ -6657,6 +7369,10 @@ struct sqlite3_module {
   /* The methods above are in versions 1 and 2 of the sqlite_module object.
   ** Those below are for version 3 and greater. */
   int (*xShadowName)(const char*);
+  /* The methods above are in versions 1 through 3 of the sqlite_module object.
+  ** Those below are for version 4 and greater. */
+  int (*xIntegrity)(sqlite3_vtab *pVTab, const char *zSchema,
+                    const char *zTabName, int mFlags, char **pzErr);
 };
 
 /*
@@ -6699,7 +7415,7 @@ struct sqlite3_module {
 ** required by SQLite. If the table has at least 64 columns and any column
 ** to the right of the first 63 is required, then bit 63 of colUsed is also
 ** set. In other words, column iCol may be required if the expression
-** (colUsed & ((sqlite3_uint64)1 << (iCol>=63 ? 63 : iCol))) evaluates to 
+** (colUsed & ((sqlite3_uint64)1 << (iCol>=63 ? 63 : iCol))) evaluates to
 ** non-zero.
 **
 ** The [xBestIndex] method must fill aConstraintUsage[] with information
@@ -6715,10 +7431,10 @@ struct sqlite3_module {
 ** when the omit flag is true there is no guarantee that the constraint will
 ** not be checked again using byte code.)^
 **
-** ^The idxNum and idxPtr values are recorded and passed into the
+** ^The idxNum and idxStr values are recorded and passed into the
 ** [xFilter] method.
-** ^[sqlite3_free()] is used to free idxPtr if and only if
-** needToFreeIdxPtr is true.
+** ^[sqlite3_free()] is used to free idxStr if and only if
+** needToFreeIdxStr is true.
 **
 ** ^The orderByConsumed means that output from [xFilter]/[xNext] will occur in
 ** the correct order to satisfy the ORDER BY clause so that no separate
@@ -6726,17 +7442,19 @@ struct sqlite3_module {
 **
 ** ^The estimatedCost value is an estimate of the cost of a particular
 ** strategy. A cost of N indicates that the cost of the strategy is similar
-** to a linear scan of an SQLite table with N rows. A cost of log(N) 
+** to a linear scan of an SQLite table with N rows. A cost of log(N)
 ** indicates that the expense of the operation is similar to that of a
 ** binary search on a unique indexed field of an SQLite table with N rows.
 **
 ** ^The estimatedRows value is an estimate of the number of rows that
 ** will be returned by the strategy.
 **
-** The xBestIndex method may optionally populate the idxFlags field with a 
-** mask of SQLITE_INDEX_SCAN_* flags. Currently there is only one such flag -
-** SQLITE_INDEX_SCAN_UNIQUE. If the xBestIndex method sets this flag, SQLite
-** assumes that the strategy may visit at most one row. 
+** The xBestIndex method may optionally populate the idxFlags field with a
+** mask of SQLITE_INDEX_SCAN_* flags. One such flag is
+** [SQLITE_INDEX_SCAN_HEX], which if set causes the [EXPLAIN QUERY PLAN]
+** output to show the idxNum has hex instead of as decimal.  Another flag is
+** SQLITE_INDEX_SCAN_UNIQUE, which if set indicates that the query plan will
+** return at most one row.
 **
 ** Additionally, if xBestIndex sets the SQLITE_INDEX_SCAN_UNIQUE flag, then
 ** SQLite also assumes that if a call to the xUpdate() method is made as
@@ -6749,14 +7467,14 @@ struct sqlite3_module {
 ** the xUpdate method are automatically rolled back by SQLite.
 **
 ** IMPORTANT: The estimatedRows field was added to the sqlite3_index_info
-** structure for SQLite [version 3.8.2] ([dateof:3.8.2]). 
+** structure for SQLite [version 3.8.2] ([dateof:3.8.2]).
 ** If a virtual table extension is
-** used with an SQLite version earlier than 3.8.2, the results of attempting 
-** to read or write the estimatedRows field are undefined (but are likely 
+** used with an SQLite version earlier than 3.8.2, the results of attempting
+** to read or write the estimatedRows field are undefined (but are likely
 ** to include crashing the application). The estimatedRows field should
 ** therefore only be used if [sqlite3_libversion_number()] returns a
 ** value greater than or equal to 3008002. Similarly, the idxFlags field
-** was added for [version 3.9.0] ([dateof:3.9.0]). 
+** was added for [version 3.9.0] ([dateof:3.9.0]).
 ** It may therefore only be used if
 ** sqlite3_libversion_number() returns a value greater than or equal to
 ** 3009000.
@@ -6796,35 +7514,69 @@ struct sqlite3_index_info {
 /*
 ** CAPI3REF: Virtual Table Scan Flags
 **
-** Virtual table implementations are allowed to set the 
+** Virtual table implementations are allowed to set the
 ** [sqlite3_index_info].idxFlags field to some combination of
 ** these bits.
 */
-#define SQLITE_INDEX_SCAN_UNIQUE      1     /* Scan visits at most 1 row */
+#define SQLITE_INDEX_SCAN_UNIQUE 0x00000001 /* Scan visits at most 1 row */
+#define SQLITE_INDEX_SCAN_HEX    0x00000002 /* Display idxNum as hex */
+                                            /* in EXPLAIN QUERY PLAN */
 
 /*
 ** CAPI3REF: Virtual Table Constraint Operator Codes
 **
 ** These macros define the allowed values for the
 ** [sqlite3_index_info].aConstraint[].op field.  Each value represents
-** an operator that is part of a constraint term in the wHERE clause of
+** an operator that is part of a constraint term in the WHERE clause of
 ** a query that uses a [virtual table].
-*/
-#define SQLITE_INDEX_CONSTRAINT_EQ         2
-#define SQLITE_INDEX_CONSTRAINT_GT         4
-#define SQLITE_INDEX_CONSTRAINT_LE         8
-#define SQLITE_INDEX_CONSTRAINT_LT        16
-#define SQLITE_INDEX_CONSTRAINT_GE        32
-#define SQLITE_INDEX_CONSTRAINT_MATCH     64
-#define SQLITE_INDEX_CONSTRAINT_LIKE      65
-#define SQLITE_INDEX_CONSTRAINT_GLOB      66
-#define SQLITE_INDEX_CONSTRAINT_REGEXP    67
-#define SQLITE_INDEX_CONSTRAINT_NE        68
-#define SQLITE_INDEX_CONSTRAINT_ISNOT     69
-#define SQLITE_INDEX_CONSTRAINT_ISNOTNULL 70
-#define SQLITE_INDEX_CONSTRAINT_ISNULL    71
-#define SQLITE_INDEX_CONSTRAINT_IS        72
-#define SQLITE_INDEX_CONSTRAINT_FUNCTION 150
+**
+** ^The left-hand operand of the operator is given by the corresponding
+** aConstraint[].iColumn field.  ^An iColumn of -1 indicates the left-hand
+** operand is the rowid.
+** The SQLITE_INDEX_CONSTRAINT_LIMIT and SQLITE_INDEX_CONSTRAINT_OFFSET
+** operators have no left-hand operand, and so for those operators the
+** corresponding aConstraint[].iColumn is meaningless and should not be
+** used.
+**
+** All operator values from SQLITE_INDEX_CONSTRAINT_FUNCTION through
+** value 255 are reserved to represent functions that are overloaded
+** by the [xFindFunction|xFindFunction method] of the virtual table
+** implementation.
+**
+** The right-hand operands for each constraint might be accessible using
+** the [sqlite3_vtab_rhs_value()] interface.  Usually the right-hand
+** operand is only available if it appears as a single constant literal
+** in the input SQL.  If the right-hand operand is another column or an
+** expression (even a constant expression) or a parameter, then the
+** sqlite3_vtab_rhs_value() probably will not be able to extract it.
+** ^The SQLITE_INDEX_CONSTRAINT_ISNULL and
+** SQLITE_INDEX_CONSTRAINT_ISNOTNULL operators have no right-hand operand
+** and hence calls to sqlite3_vtab_rhs_value() for those operators will
+** always return SQLITE_NOTFOUND.
+**
+** The collating sequence to be used for comparison can be found using
+** the [sqlite3_vtab_collation()] interface.  For most real-world virtual
+** tables, the collating sequence of constraints does not matter (for example
+** because the constraints are numeric) and so the sqlite3_vtab_collation()
+** interface is not commonly needed.
+*/
+#define SQLITE_INDEX_CONSTRAINT_EQ          2
+#define SQLITE_INDEX_CONSTRAINT_GT          4
+#define SQLITE_INDEX_CONSTRAINT_LE          8
+#define SQLITE_INDEX_CONSTRAINT_LT         16
+#define SQLITE_INDEX_CONSTRAINT_GE         32
+#define SQLITE_INDEX_CONSTRAINT_MATCH      64
+#define SQLITE_INDEX_CONSTRAINT_LIKE       65
+#define SQLITE_INDEX_CONSTRAINT_GLOB       66
+#define SQLITE_INDEX_CONSTRAINT_REGEXP     67
+#define SQLITE_INDEX_CONSTRAINT_NE         68
+#define SQLITE_INDEX_CONSTRAINT_ISNOT      69
+#define SQLITE_INDEX_CONSTRAINT_ISNOTNULL  70
+#define SQLITE_INDEX_CONSTRAINT_ISNULL     71
+#define SQLITE_INDEX_CONSTRAINT_IS         72
+#define SQLITE_INDEX_CONSTRAINT_LIMIT      73
+#define SQLITE_INDEX_CONSTRAINT_OFFSET     74
+#define SQLITE_INDEX_CONSTRAINT_FUNCTION  150
 
 /*
 ** CAPI3REF: Register A Virtual Table Implementation
@@ -6836,7 +7588,7 @@ struct sqlite3_index_info {
 ** preexisting [virtual table] for the module.
 **
 ** ^The module name is registered on the [database connection] specified
-** by the first parameter.  ^The name of the module is given by the 
+** by the first parameter.  ^The name of the module is given by the
 ** second parameter.  ^The third parameter is a pointer to
 ** the implementation of the [virtual table module].   ^The fourth
 ** parameter is an arbitrary client data pointer that is passed through
@@ -6853,7 +7605,7 @@ struct sqlite3_index_info {
 ** destructor.
 **
 ** ^If the third parameter (the pointer to the sqlite3_module object) is
-** NULL then no new module is create and any existing modules with the
+** NULL then no new module is created and any existing modules with the
 ** same name are dropped.
 **
 ** See also: [sqlite3_drop_modules()]
@@ -6951,7 +7703,7 @@ SQLITE_API int sqlite3_declare_vtab(sqlite3*, const char *zSQL);
 ** METHOD: sqlite3
 **
 ** ^(Virtual tables can provide alternative implementations of functions
-** using the [xFindFunction] method of the [virtual table module].  
+** using the [xFindFunction] method of the [virtual table module].
 ** But global versions of those functions
 ** must exist in order to be overloaded.)^
 **
@@ -6965,16 +7717,6 @@ SQLITE_API int sqlite3_declare_vtab(sqlite3*, const char *zSQL);
 */
 SQLITE_API int sqlite3_overload_function(sqlite3*, const char *zFuncName, int nArg);
 
-/*
-** The interface to the virtual-table mechanism defined above (back up
-** to a comment remarkably similar to this one) is currently considered
-** to be experimental.  The interface might change in incompatible ways.
-** If this is a problem for you, do not use the interface at this time.
-**
-** When the virtual-table mechanism stabilizes, we will declare the
-** interface fixed, support it indefinitely, and remove this comment.
-*/
-
 /*
 ** CAPI3REF: A Handle To An Open BLOB
 ** KEYWORDS: {BLOB handle} {BLOB handles}
@@ -7002,7 +7744,7 @@ typedef struct sqlite3_blob sqlite3_blob;
 **     SELECT zColumn FROM zDb.zTable WHERE [rowid] = iRow;
 ** </pre>)^
 **
-** ^(Parameter zDb is not the filename that contains the database, but 
+** ^(Parameter zDb is not the filename that contains the database, but
 ** rather the symbolic name of the database. For attached databases, this is
 ** the name that appears after the AS keyword in the [ATTACH] statement.
 ** For the main database file, the database name is "main". For TEMP
@@ -7015,28 +7757,28 @@ typedef struct sqlite3_blob sqlite3_blob;
 ** ^(On success, [SQLITE_OK] is returned and the new [BLOB handle] is stored
 ** in *ppBlob. Otherwise an [error code] is returned and, unless the error
 ** code is SQLITE_MISUSE, *ppBlob is set to NULL.)^ ^This means that, provided
-** the API is not misused, it is always safe to call [sqlite3_blob_close()] 
+** the API is not misused, it is always safe to call [sqlite3_blob_close()]
 ** on *ppBlob after this function it returns.
 **
 ** This function fails with SQLITE_ERROR if any of the following are true:
 ** <ul>
-**   <li> ^(Database zDb does not exist)^, 
-**   <li> ^(Table zTable does not exist within database zDb)^, 
-**   <li> ^(Table zTable is a WITHOUT ROWID table)^, 
+**   <li> ^(Database zDb does not exist)^,
+**   <li> ^(Table zTable does not exist within database zDb)^,
+**   <li> ^(Table zTable is a WITHOUT ROWID table)^,
 **   <li> ^(Column zColumn does not exist)^,
 **   <li> ^(Row iRow is not present in the table)^,
 **   <li> ^(The specified column of row iRow contains a value that is not
 **         a TEXT or BLOB value)^,
-**   <li> ^(Column zColumn is part of an index, PRIMARY KEY or UNIQUE 
+**   <li> ^(Column zColumn is part of an index, PRIMARY KEY or UNIQUE
 **         constraint and the blob is being opened for read/write access)^,
-**   <li> ^([foreign key constraints | Foreign key constraints] are enabled, 
+**   <li> ^([foreign key constraints | Foreign key constraints] are enabled,
 **         column zColumn is part of a [child key] definition and the blob is
 **         being opened for read/write access)^.
 ** </ul>
 **
-** ^Unless it returns SQLITE_MISUSE, this function sets the 
-** [database connection] error code and message accessible via 
-** [sqlite3_errcode()] and [sqlite3_errmsg()] and related functions. 
+** ^Unless it returns SQLITE_MISUSE, this function sets the
+** [database connection] error code and message accessible via
+** [sqlite3_errcode()] and [sqlite3_errmsg()] and related functions.
 **
 ** A BLOB referenced by sqlite3_blob_open() may be read using the
 ** [sqlite3_blob_read()] interface and modified by using
@@ -7062,7 +7804,7 @@ typedef struct sqlite3_blob sqlite3_blob;
 ** blob.
 **
 ** ^The [sqlite3_bind_zeroblob()] and [sqlite3_result_zeroblob()] interfaces
-** and the built-in [zeroblob] SQL function may be used to create a 
+** and the built-in [zeroblob] SQL function may be used to create a
 ** zero-filled blob to read or write using the incremental-blob interface.
 **
 ** To avoid a resource leak, every open [BLOB handle] should eventually
@@ -7112,7 +7854,7 @@ SQLITE_API int sqlite3_blob_reopen(sqlite3_blob *, sqlite3_int64);
 ** DESTRUCTOR: sqlite3_blob
 **
 ** ^This function closes an open [BLOB handle]. ^(The BLOB handle is closed
-** unconditionally.  Even if this routine returns an error code, the 
+** unconditionally.  Even if this routine returns an error code, the
 ** handle is still closed.)^
 **
 ** ^If the blob handle being closed was opened for read-write access, and if
@@ -7122,10 +7864,10 @@ SQLITE_API int sqlite3_blob_reopen(sqlite3_blob *, sqlite3_int64);
 ** code is returned and the transaction rolled back.
 **
 ** Calling this function with an argument that is not a NULL pointer or an
-** open blob handle results in undefined behaviour. ^Calling this routine 
-** with a null pointer (such as would be returned by a failed call to 
+** open blob handle results in undefined behavior. ^Calling this routine
+** with a null pointer (such as would be returned by a failed call to
 ** [sqlite3_blob_open()]) is a harmless no-op. ^Otherwise, if this function
-** is passed a valid open blob handle, the values returned by the 
+** is passed a valid open blob handle, the values returned by the
 ** sqlite3_errcode() and sqlite3_errmsg() functions are set before returning.
 */
 SQLITE_API int sqlite3_blob_close(sqlite3_blob *);
@@ -7134,7 +7876,7 @@ SQLITE_API int sqlite3_blob_close(sqlite3_blob *);
 ** CAPI3REF: Return The Size Of An Open BLOB
 ** METHOD: sqlite3_blob
 **
-** ^Returns the size in bytes of the BLOB accessible via the 
+** ^Returns the size in bytes of the BLOB accessible via the
 ** successfully opened [BLOB handle] in its only argument.  ^The
 ** incremental blob I/O routines can only read or overwriting existing
 ** blob content; they cannot change the size of a blob.
@@ -7185,9 +7927,9 @@ SQLITE_API int sqlite3_blob_read(sqlite3_blob *, void *Z, int N, int iOffset);
 **
 ** ^(On success, sqlite3_blob_write() returns SQLITE_OK.
 ** Otherwise, an  [error code] or an [extended error code] is returned.)^
-** ^Unless SQLITE_MISUSE is returned, this function sets the 
-** [database connection] error code and message accessible via 
-** [sqlite3_errcode()] and [sqlite3_errmsg()] and related functions. 
+** ^Unless SQLITE_MISUSE is returned, this function sets the
+** [database connection] error code and message accessible via
+** [sqlite3_errcode()] and [sqlite3_errmsg()] and related functions.
 **
 ** ^If the [BLOB handle] passed as the first argument was not opened for
 ** writing (the flags parameter to [sqlite3_blob_open()] was zero),
@@ -7196,9 +7938,9 @@ SQLITE_API int sqlite3_blob_read(sqlite3_blob *, void *Z, int N, int iOffset);
 ** This function may only modify the contents of the BLOB; it is
 ** not possible to increase the size of a BLOB using this API.
 ** ^If offset iOffset is less than N bytes from the end of the BLOB,
-** [SQLITE_ERROR] is returned and no data is written. The size of the 
-** BLOB (and hence the maximum value of N+iOffset) can be determined 
-** using the [sqlite3_blob_bytes()] interface. ^If N or iOffset are less 
+** [SQLITE_ERROR] is returned and no data is written. The size of the
+** BLOB (and hence the maximum value of N+iOffset) can be determined
+** using the [sqlite3_blob_bytes()] interface. ^If N or iOffset are less
 ** than zero [SQLITE_ERROR] is returned and no data is written.
 **
 ** ^An attempt to write to an expired [BLOB handle] fails with an
@@ -7292,7 +8034,7 @@ SQLITE_API int sqlite3_vfs_unregister(sqlite3_vfs*);
 ** <ul>
 ** <li>  SQLITE_MUTEX_FAST
 ** <li>  SQLITE_MUTEX_RECURSIVE
-** <li>  SQLITE_MUTEX_STATIC_MASTER
+** <li>  SQLITE_MUTEX_STATIC_MAIN
 ** <li>  SQLITE_MUTEX_STATIC_MEM
 ** <li>  SQLITE_MUTEX_STATIC_OPEN
 ** <li>  SQLITE_MUTEX_STATIC_PRNG
@@ -7349,18 +8091,20 @@ SQLITE_API int sqlite3_vfs_unregister(sqlite3_vfs*);
 **
 ** ^(Some systems (for example, Windows 95) do not support the operation
 ** implemented by sqlite3_mutex_try().  On those systems, sqlite3_mutex_try()
-** will always return SQLITE_BUSY. The SQLite core only ever uses
-** sqlite3_mutex_try() as an optimization so this is acceptable 
-** behavior.)^
+** will always return SQLITE_BUSY. In most cases the SQLite core only uses
+** sqlite3_mutex_try() as an optimization, so this is acceptable
+** behavior. The exceptions are unix builds that set the
+** SQLITE_ENABLE_SETLK_TIMEOUT build option. In that case a working
+** sqlite3_mutex_try() is required.)^
 **
 ** ^The sqlite3_mutex_leave() routine exits a mutex that was
 ** previously entered by the same thread.   The behavior
 ** is undefined if the mutex is not currently entered by the
 ** calling thread or is not currently allocated.
 **
-** ^If the argument to sqlite3_mutex_enter(), sqlite3_mutex_try(), or
-** sqlite3_mutex_leave() is a NULL pointer, then all three routines
-** behave as no-ops.
+** ^If the argument to sqlite3_mutex_enter(), sqlite3_mutex_try(),
+** sqlite3_mutex_leave(), or sqlite3_mutex_free() is a NULL pointer,
+** then any of the four routines behaves as a no-op.
 **
 ** See also: [sqlite3_mutex_held()] and [sqlite3_mutex_notheld()].
 */
@@ -7494,7 +8238,7 @@ SQLITE_API int sqlite3_mutex_notheld(sqlite3_mutex*);
 */
 #define SQLITE_MUTEX_FAST             0
 #define SQLITE_MUTEX_RECURSIVE        1
-#define SQLITE_MUTEX_STATIC_MASTER    2
+#define SQLITE_MUTEX_STATIC_MAIN      2
 #define SQLITE_MUTEX_STATIC_MEM       3  /* sqlite3_malloc() */
 #define SQLITE_MUTEX_STATIC_MEM2      4  /* NOT USED */
 #define SQLITE_MUTEX_STATIC_OPEN      4  /* sqlite3BtreeOpen() */
@@ -7509,11 +8253,15 @@ SQLITE_API int sqlite3_mutex_notheld(sqlite3_mutex*);
 #define SQLITE_MUTEX_STATIC_VFS2     12  /* For use by extension VFS */
 #define SQLITE_MUTEX_STATIC_VFS3     13  /* For use by application VFS */
 
+/* Legacy compatibility: */
+#define SQLITE_MUTEX_STATIC_MASTER    2
+
+
 /*
 ** CAPI3REF: Retrieve the mutex for a database connection
 ** METHOD: sqlite3
 **
-** ^This interface returns a pointer the [sqlite3_mutex] object that 
+** ^This interface returns a pointer the [sqlite3_mutex] object that
 ** serializes access to the [database connection] given in the argument
 ** when the [threading mode] is Serialized.
 ** ^If the [threading mode] is Single-thread or Multi-thread then this
@@ -7540,7 +8288,7 @@ SQLITE_API sqlite3_mutex *sqlite3_db_mutex(sqlite3*);
 ** method becomes the return value of this routine.
 **
 ** A few opcodes for [sqlite3_file_control()] are handled directly
-** by the SQLite core and never invoke the 
+** by the SQLite core and never invoke the
 ** sqlite3_io_methods.xFileControl method.
 ** ^The [SQLITE_FCNTL_FILE_POINTER] value for the op parameter causes
 ** a pointer to the underlying [sqlite3_file] object to be written into
@@ -7598,15 +8346,18 @@ SQLITE_API int sqlite3_test_control(int op, ...);
 #define SQLITE_TESTCTRL_PRNG_SAVE                5
 #define SQLITE_TESTCTRL_PRNG_RESTORE             6
 #define SQLITE_TESTCTRL_PRNG_RESET               7  /* NOT USED */
+#define SQLITE_TESTCTRL_FK_NO_ACTION             7
 #define SQLITE_TESTCTRL_BITVEC_TEST              8
 #define SQLITE_TESTCTRL_FAULT_INSTALL            9
 #define SQLITE_TESTCTRL_BENIGN_MALLOC_HOOKS     10
 #define SQLITE_TESTCTRL_PENDING_BYTE            11
 #define SQLITE_TESTCTRL_ASSERT                  12
 #define SQLITE_TESTCTRL_ALWAYS                  13
-#define SQLITE_TESTCTRL_RESERVE                 14
+#define SQLITE_TESTCTRL_RESERVE                 14  /* NOT USED */
+#define SQLITE_TESTCTRL_JSON_SELFCHECK          14
 #define SQLITE_TESTCTRL_OPTIMIZATIONS           15
 #define SQLITE_TESTCTRL_ISKEYWORD               16  /* NOT USED */
+#define SQLITE_TESTCTRL_GETOPT                  16
 #define SQLITE_TESTCTRL_SCRATCHMALLOC           17  /* NOT USED */
 #define SQLITE_TESTCTRL_INTERNAL_FUNCTIONS      17
 #define SQLITE_TESTCTRL_LOCALTIME_FAULT         18
@@ -7622,12 +8373,17 @@ SQLITE_API int sqlite3_test_control(int op, ...);
 #define SQLITE_TESTCTRL_RESULT_INTREAL          27
 #define SQLITE_TESTCTRL_PRNG_SEED               28
 #define SQLITE_TESTCTRL_EXTRA_SCHEMA_CHECKS     29
-#define SQLITE_TESTCTRL_LAST                    29  /* Largest TESTCTRL */
+#define SQLITE_TESTCTRL_SEEK_COUNT              30
+#define SQLITE_TESTCTRL_TRACEFLAGS              31
+#define SQLITE_TESTCTRL_TUNE                    32
+#define SQLITE_TESTCTRL_LOGEST                  33
+#define SQLITE_TESTCTRL_USELONGDOUBLE           34  /* NOT USED */
+#define SQLITE_TESTCTRL_LAST                    34  /* Largest TESTCTRL */
 
 /*
 ** CAPI3REF: SQL Keyword Checking
 **
-** These routines provide access to the set of SQL language keywords 
+** These routines provide access to the set of SQL language keywords
 ** recognized by SQLite.  Applications can uses these routines to determine
 ** whether or not a specific identifier needs to be escaped (for example,
 ** by enclosing in double-quotes) so as not to confuse the parser.
@@ -7635,7 +8391,7 @@ SQLITE_API int sqlite3_test_control(int op, ...);
 ** The sqlite3_keyword_count() interface returns the number of distinct
 ** keywords understood by SQLite.
 **
-** The sqlite3_keyword_name(N,Z,L) interface finds the N-th keyword and
+** The sqlite3_keyword_name(N,Z,L) interface finds the 0-based N-th keyword and
 ** makes *Z point to that keyword expressed as UTF8 and writes the number
 ** of bytes in the keyword into *L.  The string that *Z points to is not
 ** zero-terminated.  The sqlite3_keyword_name(N,Z,L) routine returns
@@ -7699,14 +8455,14 @@ typedef struct sqlite3_str sqlite3_str;
 **
 ** ^The [sqlite3_str_new(D)] interface allocates and initializes
 ** a new [sqlite3_str] object.  To avoid memory leaks, the object returned by
-** [sqlite3_str_new()] must be freed by a subsequent call to 
+** [sqlite3_str_new()] must be freed by a subsequent call to
 ** [sqlite3_str_finish(X)].
 **
 ** ^The [sqlite3_str_new(D)] interface always returns a pointer to a
 ** valid [sqlite3_str] object, though in the event of an out-of-memory
 ** error the returned object might be a special singleton that will
-** silently reject new text, always return SQLITE_NOMEM from 
-** [sqlite3_str_errcode()], always return 0 for 
+** silently reject new text, always return SQLITE_NOMEM from
+** [sqlite3_str_errcode()], always return 0 for
 ** [sqlite3_str_length()], and always return NULL from
 ** [sqlite3_str_finish(X)].  It is always safe to use the value
 ** returned by [sqlite3_str_new(D)] as the sqlite3_str parameter
@@ -7742,9 +8498,9 @@ SQLITE_API char *sqlite3_str_finish(sqlite3_str*);
 ** These interfaces add content to an sqlite3_str object previously obtained
 ** from [sqlite3_str_new()].
 **
-** ^The [sqlite3_str_appendf(X,F,...)] and 
+** ^The [sqlite3_str_appendf(X,F,...)] and
 ** [sqlite3_str_vappendf(X,F,V)] interfaces uses the [built-in printf]
-** functionality of SQLite to append formatted text onto the end of 
+** functionality of SQLite to append formatted text onto the end of
 ** [sqlite3_str] object X.
 **
 ** ^The [sqlite3_str_append(X,S,N)] method appends exactly N bytes from string S
@@ -7761,7 +8517,7 @@ SQLITE_API char *sqlite3_str_finish(sqlite3_str*);
 ** ^This method can be used, for example, to add whitespace indentation.
 **
 ** ^The [sqlite3_str_reset(X)] method resets the string under construction
-** inside [sqlite3_str] object X back to zero bytes in length.  
+** inside [sqlite3_str] object X back to zero bytes in length.
 **
 ** These methods do not return a result code.  ^If an error occurs, that fact
 ** is recorded in the [sqlite3_str] object and can be recovered by a
@@ -7863,7 +8619,7 @@ SQLITE_API int sqlite3_status64(
 ** <dd>This parameter records the largest memory allocation request
 ** handed to [sqlite3_malloc()] or [sqlite3_realloc()] (or their
 ** internal equivalents).  Only the value returned in the
-** *pHighwater parameter to [sqlite3_status()] is of interest.  
+** *pHighwater parameter to [sqlite3_status()] is of interest.
 ** The value written into the *pCurrent parameter is undefined.</dd>)^
 **
 ** [[SQLITE_STATUS_MALLOC_COUNT]] ^(<dt>SQLITE_STATUS_MALLOC_COUNT</dt>
@@ -7872,11 +8628,11 @@ SQLITE_API int sqlite3_status64(
 **
 ** [[SQLITE_STATUS_PAGECACHE_USED]] ^(<dt>SQLITE_STATUS_PAGECACHE_USED</dt>
 ** <dd>This parameter returns the number of pages used out of the
-** [pagecache memory allocator] that was configured using 
+** [pagecache memory allocator] that was configured using
 ** [SQLITE_CONFIG_PAGECACHE].  The
 ** value returned is in pages, not in bytes.</dd>)^
 **
-** [[SQLITE_STATUS_PAGECACHE_OVERFLOW]] 
+** [[SQLITE_STATUS_PAGECACHE_OVERFLOW]]
 ** ^(<dt>SQLITE_STATUS_PAGECACHE_OVERFLOW</dt>
 ** <dd>This parameter returns the number of bytes of page cache
 ** allocation which could not be satisfied by the [SQLITE_CONFIG_PAGECACHE]
@@ -7889,7 +8645,7 @@ SQLITE_API int sqlite3_status64(
 ** [[SQLITE_STATUS_PAGECACHE_SIZE]] ^(<dt>SQLITE_STATUS_PAGECACHE_SIZE</dt>
 ** <dd>This parameter records the largest memory allocation request
 ** handed to the [pagecache memory allocator].  Only the value returned in the
-** *pHighwater parameter to [sqlite3_status()] is of interest.  
+** *pHighwater parameter to [sqlite3_status()] is of interest.
 ** The value written into the *pCurrent parameter is undefined.</dd>)^
 **
 ** [[SQLITE_STATUS_SCRATCH_USED]] <dt>SQLITE_STATUS_SCRATCH_USED</dt>
@@ -7902,7 +8658,7 @@ SQLITE_API int sqlite3_status64(
 ** <dd>No longer used.</dd>
 **
 ** [[SQLITE_STATUS_PARSER_STACK]] ^(<dt>SQLITE_STATUS_PARSER_STACK</dt>
-** <dd>The *pHighwater parameter records the deepest parser stack. 
+** <dd>The *pHighwater parameter records the deepest parser stack.
 ** The *pCurrent value is undefined.  The *pHighwater value is only
 ** meaningful if SQLite is compiled with [YYTRACKMAXSTACKDEPTH].</dd>)^
 ** </dl>
@@ -7924,12 +8680,12 @@ SQLITE_API int sqlite3_status64(
 ** CAPI3REF: Database Connection Status
 ** METHOD: sqlite3
 **
-** ^This interface is used to retrieve runtime status information 
+** ^This interface is used to retrieve runtime status information
 ** about a single [database connection].  ^The first argument is the
 ** database connection object to be interrogated.  ^The second argument
 ** is an integer constant, taken from the set of
 ** [SQLITE_DBSTATUS options], that
-** determines the parameter to interrogate.  The set of 
+** determines the parameter to interrogate.  The set of
 ** [SQLITE_DBSTATUS options] is likely
 ** to grow in future releases of SQLite.
 **
@@ -7964,7 +8720,7 @@ SQLITE_API int sqlite3_db_status(sqlite3*, int op, int *pCur, int *pHiwtr, int r
 ** checked out.</dd>)^
 **
 ** [[SQLITE_DBSTATUS_LOOKASIDE_HIT]] ^(<dt>SQLITE_DBSTATUS_LOOKASIDE_HIT</dt>
-** <dd>This parameter returns the number of malloc attempts that were 
+** <dd>This parameter returns the number of malloc attempts that were
 ** satisfied using lookaside memory. Only the high-water value is meaningful;
 ** the current value is always zero.)^
 **
@@ -7989,7 +8745,7 @@ SQLITE_API int sqlite3_db_status(sqlite3*, int op, int *pCur, int *pHiwtr, int r
 ** memory used by all pager caches associated with the database connection.)^
 ** ^The highwater mark associated with SQLITE_DBSTATUS_CACHE_USED is always 0.
 **
-** [[SQLITE_DBSTATUS_CACHE_USED_SHARED]] 
+** [[SQLITE_DBSTATUS_CACHE_USED_SHARED]]
 ** ^(<dt>SQLITE_DBSTATUS_CACHE_USED_SHARED</dt>
 ** <dd>This parameter is similar to DBSTATUS_CACHE_USED, except that if a
 ** pager cache is shared between two or more connections the bytes of heap
@@ -8004,7 +8760,7 @@ SQLITE_API int sqlite3_db_status(sqlite3*, int op, int *pCur, int *pHiwtr, int r
 ** [[SQLITE_DBSTATUS_SCHEMA_USED]] ^(<dt>SQLITE_DBSTATUS_SCHEMA_USED</dt>
 ** <dd>This parameter returns the approximate number of bytes of heap
 ** memory used to store the schema for all databases associated
-** with the connection - main, temp, and any [ATTACH]-ed databases.)^ 
+** with the connection - main, temp, and any [ATTACH]-ed databases.)^
 ** ^The full amount of memory used by the schemas is reported, even if the
 ** schema memory is shared with other database connections due to
 ** [shared cache mode] being enabled.
@@ -8019,13 +8775,13 @@ SQLITE_API int sqlite3_db_status(sqlite3*, int op, int *pCur, int *pHiwtr, int r
 **
 ** [[SQLITE_DBSTATUS_CACHE_HIT]] ^(<dt>SQLITE_DBSTATUS_CACHE_HIT</dt>
 ** <dd>This parameter returns the number of pager cache hits that have
-** occurred.)^ ^The highwater mark associated with SQLITE_DBSTATUS_CACHE_HIT 
+** occurred.)^ ^The highwater mark associated with SQLITE_DBSTATUS_CACHE_HIT
 ** is always 0.
 ** </dd>
 **
 ** [[SQLITE_DBSTATUS_CACHE_MISS]] ^(<dt>SQLITE_DBSTATUS_CACHE_MISS</dt>
 ** <dd>This parameter returns the number of pager cache misses that have
-** occurred.)^ ^The highwater mark associated with SQLITE_DBSTATUS_CACHE_MISS 
+** occurred.)^ ^The highwater mark associated with SQLITE_DBSTATUS_CACHE_MISS
 ** is always 0.
 ** </dd>
 **
@@ -8083,7 +8839,7 @@ SQLITE_API int sqlite3_db_status(sqlite3*, int op, int *pCur, int *pHiwtr, int r
 ** statements.  For example, if the number of table steps greatly exceeds
 ** the number of table searches or result rows, that would tend to indicate
 ** that the prepared statement is using a full table scan rather than
-** an index.  
+** an index.
 **
 ** ^(This interface is used to retrieve and reset counter values from
 ** a [prepared statement].  The first argument is the prepared statement
@@ -8110,7 +8866,7 @@ SQLITE_API int sqlite3_stmt_status(sqlite3_stmt*, int op,int resetFlg);
 ** [[SQLITE_STMTSTATUS_FULLSCAN_STEP]] <dt>SQLITE_STMTSTATUS_FULLSCAN_STEP</dt>
 ** <dd>^This is the number of times that SQLite has stepped forward in
 ** a table as part of a full table scan.  Large numbers for this counter
-** may indicate opportunities for performance improvement through 
+** may indicate opportunities for performance improvement through
 ** careful use of indices.</dd>
 **
 ** [[SQLITE_STMTSTATUS_SORT]] <dt>SQLITE_STMTSTATUS_SORT</dt>
@@ -8128,14 +8884,14 @@ SQLITE_API int sqlite3_stmt_status(sqlite3_stmt*, int op,int resetFlg);
 ** [[SQLITE_STMTSTATUS_VM_STEP]] <dt>SQLITE_STMTSTATUS_VM_STEP</dt>
 ** <dd>^This is the number of virtual machine operations executed
 ** by the prepared statement if that number is less than or equal
-** to 2147483647.  The number of virtual machine operations can be 
+** to 2147483647.  The number of virtual machine operations can be
 ** used as a proxy for the total work done by the prepared statement.
 ** If the number of virtual machine operations exceeds 2147483647
 ** then the value returned by this statement status code is undefined.
 **
 ** [[SQLITE_STMTSTATUS_REPREPARE]] <dt>SQLITE_STMTSTATUS_REPREPARE</dt>
 ** <dd>^This is the number of times that the prepare statement has been
-** automatically regenerated due to schema changes or changes to 
+** automatically regenerated due to schema changes or changes to
 ** [bound parameters] that might affect the query plan.
 **
 ** [[SQLITE_STMTSTATUS_RUN]] <dt>SQLITE_STMTSTATUS_RUN</dt>
@@ -8145,6 +8901,16 @@ SQLITE_API int sqlite3_stmt_status(sqlite3_stmt*, int op,int resetFlg);
 ** The counter is incremented on the first [sqlite3_step()] call of each
 ** cycle.
 **
+** [[SQLITE_STMTSTATUS_FILTER_MISS]]
+** [[SQLITE_STMTSTATUS_FILTER HIT]]
+** <dt>SQLITE_STMTSTATUS_FILTER_HIT<br>
+** SQLITE_STMTSTATUS_FILTER_MISS</dt>
+** <dd>^SQLITE_STMTSTATUS_FILTER_HIT is the number of times that a join
+** step was bypassed because a Bloom filter returned not-found.  The
+** corresponding SQLITE_STMTSTATUS_FILTER_MISS value is the number of
+** times that the Bloom filter returned a find, and thus the join step
+** had to be processed as normal.
+**
 ** [[SQLITE_STMTSTATUS_MEMUSED]] <dt>SQLITE_STMTSTATUS_MEMUSED</dt>
 ** <dd>^This is the approximate number of bytes of heap memory
 ** used to store the prepared statement.  ^This value is not actually
@@ -8159,6 +8925,8 @@ SQLITE_API int sqlite3_stmt_status(sqlite3_stmt*, int op,int resetFlg);
 #define SQLITE_STMTSTATUS_VM_STEP           4
 #define SQLITE_STMTSTATUS_REPREPARE         5
 #define SQLITE_STMTSTATUS_RUN               6
+#define SQLITE_STMTSTATUS_FILTER_MISS       7
+#define SQLITE_STMTSTATUS_FILTER_HIT        8
 #define SQLITE_STMTSTATUS_MEMUSED           99
 
 /*
@@ -8195,15 +8963,15 @@ struct sqlite3_pcache_page {
 ** KEYWORDS: {page cache}
 **
 ** ^(The [sqlite3_config]([SQLITE_CONFIG_PCACHE2], ...) interface can
-** register an alternative page cache implementation by passing in an 
+** register an alternative page cache implementation by passing in an
 ** instance of the sqlite3_pcache_methods2 structure.)^
-** In many applications, most of the heap memory allocated by 
+** In many applications, most of the heap memory allocated by
 ** SQLite is used for the page cache.
-** By implementing a 
+** By implementing a
 ** custom page cache using this API, an application can better control
-** the amount of memory consumed by SQLite, the way in which 
-** that memory is allocated and released, and the policies used to 
-** determine exactly which parts of a database file are cached and for 
+** the amount of memory consumed by SQLite, the way in which
+** that memory is allocated and released, and the policies used to
+** determine exactly which parts of a database file are cached and for
 ** how long.
 **
 ** The alternative page cache mechanism is an
@@ -8216,19 +8984,19 @@ struct sqlite3_pcache_page {
 ** [sqlite3_config()] returns.)^
 **
 ** [[the xInit() page cache method]]
-** ^(The xInit() method is called once for each effective 
+** ^(The xInit() method is called once for each effective
 ** call to [sqlite3_initialize()])^
 ** (usually only once during the lifetime of the process). ^(The xInit()
 ** method is passed a copy of the sqlite3_pcache_methods2.pArg value.)^
-** The intent of the xInit() method is to set up global data structures 
-** required by the custom page cache implementation. 
-** ^(If the xInit() method is NULL, then the 
+** The intent of the xInit() method is to set up global data structures
+** required by the custom page cache implementation.
+** ^(If the xInit() method is NULL, then the
 ** built-in default page cache is used instead of the application defined
 ** page cache.)^
 **
 ** [[the xShutdown() page cache method]]
 ** ^The xShutdown() method is called by [sqlite3_shutdown()].
-** It can be used to clean up 
+** It can be used to clean up
 ** any outstanding resources before process shutdown, if required.
 ** ^The xShutdown() method may be NULL.
 **
@@ -8247,7 +9015,7 @@ struct sqlite3_pcache_page {
 ** though this is not guaranteed. ^The
 ** first parameter, szPage, is the size in bytes of the pages that must
 ** be allocated by the cache.  ^szPage will always a power of two.  ^The
-** second parameter szExtra is a number of bytes of extra storage 
+** second parameter szExtra is a number of bytes of extra storage
 ** associated with each page cache entry.  ^The szExtra parameter will
 ** a number less than 250.  SQLite will use the
 ** extra szExtra bytes on each page to store metadata about the underlying
@@ -8260,7 +9028,7 @@ struct sqlite3_pcache_page {
 ** it is purely advisory.  ^On a cache where bPurgeable is false, SQLite will
 ** never invoke xUnpin() except to deliberately delete a page.
 ** ^In other words, calls to xUnpin() on a cache with bPurgeable set to
-** false will always have the "discard" flag set to true.  
+** false will always have the "discard" flag set to true.
 ** ^Hence, a cache created with bPurgeable false will
 ** never contain any unpinned pages.
 **
@@ -8275,12 +9043,12 @@ struct sqlite3_pcache_page {
 ** [[the xPagecount() page cache methods]]
 ** The xPagecount() method must return the number of pages currently
 ** stored in the cache, both pinned and unpinned.
-** 
+**
 ** [[the xFetch() page cache methods]]
-** The xFetch() method locates a page in the cache and returns a pointer to 
+** The xFetch() method locates a page in the cache and returns a pointer to
 ** an sqlite3_pcache_page object associated with that page, or a NULL pointer.
 ** The pBuf element of the returned sqlite3_pcache_page object will be a
-** pointer to a buffer of szPage bytes used to store the content of a 
+** pointer to a buffer of szPage bytes used to store the content of a
 ** single database page.  The pExtra element of sqlite3_pcache_page will be
 ** a pointer to the szExtra bytes of extra storage that SQLite has requested
 ** for each entry in the page cache.
@@ -8319,8 +9087,8 @@ struct sqlite3_pcache_page {
 ** page cache implementation. ^The page cache implementation
 ** may choose to evict unpinned pages at any time.
 **
-** The cache must not perform any reference counting. A single 
-** call to xUnpin() unpins the page regardless of the number of prior calls 
+** The cache must not perform any reference counting. A single
+** call to xUnpin() unpins the page regardless of the number of prior calls
 ** to xFetch().
 **
 ** [[the xRekey() page cache methods]]
@@ -8360,7 +9128,7 @@ struct sqlite3_pcache_methods2 {
   int (*xPagecount)(sqlite3_pcache*);
   sqlite3_pcache_page *(*xFetch)(sqlite3_pcache*, unsigned key, int createFlag);
   void (*xUnpin)(sqlite3_pcache*, sqlite3_pcache_page*, int discard);
-  void (*xRekey)(sqlite3_pcache*, sqlite3_pcache_page*, 
+  void (*xRekey)(sqlite3_pcache*, sqlite3_pcache_page*,
       unsigned oldKey, unsigned newKey);
   void (*xTruncate)(sqlite3_pcache*, unsigned iLimit);
   void (*xDestroy)(sqlite3_pcache*);
@@ -8405,7 +9173,7 @@ typedef struct sqlite3_backup sqlite3_backup;
 **
 ** The backup API copies the content of one database into another.
 ** It is useful either for creating backups of databases or
-** for copying in-memory databases to or from persistent files. 
+** for copying in-memory databases to or from persistent files.
 **
 ** See Also: [Using the SQLite Online Backup API]
 **
@@ -8416,36 +9184,36 @@ typedef struct sqlite3_backup sqlite3_backup;
 ** ^Thus, the backup may be performed on a live source database without
 ** preventing other database connections from
 ** reading or writing to the source database while the backup is underway.
-** 
-** ^(To perform a backup operation: 
+**
+** ^(To perform a backup operation:
 **   <ol>
 **     <li><b>sqlite3_backup_init()</b> is called once to initialize the
-**         backup, 
-**     <li><b>sqlite3_backup_step()</b> is called one or more times to transfer 
+**         backup,
+**     <li><b>sqlite3_backup_step()</b> is called one or more times to transfer
 **         the data between the two databases, and finally
-**     <li><b>sqlite3_backup_finish()</b> is called to release all resources 
-**         associated with the backup operation. 
+**     <li><b>sqlite3_backup_finish()</b> is called to release all resources
+**         associated with the backup operation.
 **   </ol>)^
 ** There should be exactly one call to sqlite3_backup_finish() for each
 ** successful call to sqlite3_backup_init().
 **
 ** [[sqlite3_backup_init()]] <b>sqlite3_backup_init()</b>
 **
-** ^The D and N arguments to sqlite3_backup_init(D,N,S,M) are the 
-** [database connection] associated with the destination database 
+** ^The D and N arguments to sqlite3_backup_init(D,N,S,M) are the
+** [database connection] associated with the destination database
 ** and the database name, respectively.
 ** ^The database name is "main" for the main database, "temp" for the
 ** temporary database, or the name specified after the AS keyword in
 ** an [ATTACH] statement for an attached database.
-** ^The S and M arguments passed to 
+** ^The S and M arguments passed to
 ** sqlite3_backup_init(D,N,S,M) identify the [database connection]
 ** and database name of the source database, respectively.
 ** ^The source and destination [database connections] (parameters S and D)
 ** must be different or else sqlite3_backup_init(D,N,S,M) will fail with
 ** an error.
 **
-** ^A call to sqlite3_backup_init() will fail, returning NULL, if 
-** there is already a read or read-write transaction open on the 
+** ^A call to sqlite3_backup_init() will fail, returning NULL, if
+** there is already a read or read-write transaction open on the
 ** destination database.
 **
 ** ^If an error occurs within sqlite3_backup_init(D,N,S,M), then NULL is
@@ -8457,14 +9225,14 @@ typedef struct sqlite3_backup sqlite3_backup;
 ** ^A successful call to sqlite3_backup_init() returns a pointer to an
 ** [sqlite3_backup] object.
 ** ^The [sqlite3_backup] object may be used with the sqlite3_backup_step() and
-** sqlite3_backup_finish() functions to perform the specified backup 
+** sqlite3_backup_finish() functions to perform the specified backup
 ** operation.
 **
 ** [[sqlite3_backup_step()]] <b>sqlite3_backup_step()</b>
 **
-** ^Function sqlite3_backup_step(B,N) will copy up to N pages between 
+** ^Function sqlite3_backup_step(B,N) will copy up to N pages between
 ** the source and destination databases specified by [sqlite3_backup] object B.
-** ^If N is negative, all remaining source pages are copied. 
+** ^If N is negative, all remaining source pages are copied.
 ** ^If sqlite3_backup_step(B,N) successfully copies N pages and there
 ** are still more pages to be copied, then the function returns [SQLITE_OK].
 ** ^If sqlite3_backup_step(B,N) successfully finishes copying all pages
@@ -8486,8 +9254,8 @@ typedef struct sqlite3_backup sqlite3_backup;
 **
 ** ^If sqlite3_backup_step() cannot obtain a required file-system lock, then
 ** the [sqlite3_busy_handler | busy-handler function]
-** is invoked (if one is specified). ^If the 
-** busy-handler returns non-zero before the lock is available, then 
+** is invoked (if one is specified). ^If the
+** busy-handler returns non-zero before the lock is available, then
 ** [SQLITE_BUSY] is returned to the caller. ^In this case the call to
 ** sqlite3_backup_step() can be retried later. ^If the source
 ** [database connection]
@@ -8495,15 +9263,15 @@ typedef struct sqlite3_backup sqlite3_backup;
 ** is called, then [SQLITE_LOCKED] is returned immediately. ^Again, in this
 ** case the call to sqlite3_backup_step() can be retried later on. ^(If
 ** [SQLITE_IOERR_ACCESS | SQLITE_IOERR_XXX], [SQLITE_NOMEM], or
-** [SQLITE_READONLY] is returned, then 
-** there is no point in retrying the call to sqlite3_backup_step(). These 
-** errors are considered fatal.)^  The application must accept 
-** that the backup operation has failed and pass the backup operation handle 
+** [SQLITE_READONLY] is returned, then
+** there is no point in retrying the call to sqlite3_backup_step(). These
+** errors are considered fatal.)^  The application must accept
+** that the backup operation has failed and pass the backup operation handle
 ** to the sqlite3_backup_finish() to release associated resources.
 **
 ** ^The first call to sqlite3_backup_step() obtains an exclusive lock
-** on the destination file. ^The exclusive lock is not released until either 
-** sqlite3_backup_finish() is called or the backup operation is complete 
+** on the destination file. ^The exclusive lock is not released until either
+** sqlite3_backup_finish() is called or the backup operation is complete
 ** and sqlite3_backup_step() returns [SQLITE_DONE].  ^Every call to
 ** sqlite3_backup_step() obtains a [shared lock] on the source database that
 ** lasts for the duration of the sqlite3_backup_step() call.
@@ -8512,18 +9280,18 @@ typedef struct sqlite3_backup sqlite3_backup;
 ** through the backup process.  ^If the source database is modified by an
 ** external process or via a database connection other than the one being
 ** used by the backup operation, then the backup will be automatically
-** restarted by the next call to sqlite3_backup_step(). ^If the source 
+** restarted by the next call to sqlite3_backup_step(). ^If the source
 ** database is modified by the using the same database connection as is used
 ** by the backup operation, then the backup database is automatically
 ** updated at the same time.
 **
 ** [[sqlite3_backup_finish()]] <b>sqlite3_backup_finish()</b>
 **
-** When sqlite3_backup_step() has returned [SQLITE_DONE], or when the 
+** When sqlite3_backup_step() has returned [SQLITE_DONE], or when the
 ** application wishes to abandon the backup operation, the application
 ** should destroy the [sqlite3_backup] by passing it to sqlite3_backup_finish().
 ** ^The sqlite3_backup_finish() interfaces releases all
-** resources associated with the [sqlite3_backup] object. 
+** resources associated with the [sqlite3_backup] object.
 ** ^If sqlite3_backup_step() has not yet returned [SQLITE_DONE], then any
 ** active write-transaction on the destination database is rolled back.
 ** The [sqlite3_backup] object is invalid
@@ -8563,28 +9331,38 @@ typedef struct sqlite3_backup sqlite3_backup;
 ** connections, then the source database connection may be used concurrently
 ** from within other threads.
 **
-** However, the application must guarantee that the destination 
-** [database connection] is not passed to any other API (by any thread) after 
+** However, the application must guarantee that the destination
+** [database connection] is not passed to any other API (by any thread) after
 ** sqlite3_backup_init() is called and before the corresponding call to
 ** sqlite3_backup_finish().  SQLite does not currently check to see
 ** if the application incorrectly accesses the destination [database connection]
 ** and so no error code is reported, but the operations may malfunction
 ** nevertheless.  Use of the destination database connection while a
-** backup is in progress might also also cause a mutex deadlock.
+** backup is in progress might also cause a mutex deadlock.
 **
 ** If running in [shared cache mode], the application must
 ** guarantee that the shared cache used by the destination database
 ** is not accessed while the backup is running. In practice this means
-** that the application must guarantee that the disk file being 
+** that the application must guarantee that the disk file being
 ** backed up to is not accessed by any connection within the process,
 ** not just the specific connection that was passed to sqlite3_backup_init().
 **
-** The [sqlite3_backup] object itself is partially threadsafe. Multiple 
+** The [sqlite3_backup] object itself is partially threadsafe. Multiple
 ** threads may safely make multiple concurrent calls to sqlite3_backup_step().
 ** However, the sqlite3_backup_remaining() and sqlite3_backup_pagecount()
 ** APIs are not strictly speaking threadsafe. If they are invoked at the
 ** same time as another thread is invoking sqlite3_backup_step() it is
 ** possible that they return invalid values.
+**
+** <b>Alternatives To Using The Backup API</b>
+**
+** Other techniques for safely creating a consistent backup of an SQLite
+** database include:
+**
+** <ul>
+** <li> The [VACUUM INTO] command.
+** <li> The [sqlite3_rsync] utility program.
+** </ul>
 */
 SQLITE_API sqlite3_backup *sqlite3_backup_init(
   sqlite3 *pDest,                        /* Destination database handle */
@@ -8604,8 +9382,8 @@ SQLITE_API int sqlite3_backup_pagecount(sqlite3_backup *p);
 ** ^When running in shared-cache mode, a database operation may fail with
 ** an [SQLITE_LOCKED] error if the required locks on the shared-cache or
 ** individual tables within the shared-cache cannot be obtained. See
-** [SQLite Shared-Cache Mode] for a description of shared-cache locking. 
-** ^This API may be used to register a callback that SQLite will invoke 
+** [SQLite Shared-Cache Mode] for a description of shared-cache locking.
+** ^This API may be used to register a callback that SQLite will invoke
 ** when the connection currently holding the required lock relinquishes it.
 ** ^This API is only available if the library was compiled with the
 ** [SQLITE_ENABLE_UNLOCK_NOTIFY] C-preprocessor symbol defined.
@@ -8613,14 +9391,14 @@ SQLITE_API int sqlite3_backup_pagecount(sqlite3_backup *p);
 ** See Also: [Using the SQLite Unlock Notification Feature].
 **
 ** ^Shared-cache locks are released when a database connection concludes
-** its current transaction, either by committing it or rolling it back. 
+** its current transaction, either by committing it or rolling it back.
 **
 ** ^When a connection (known as the blocked connection) fails to obtain a
 ** shared-cache lock and SQLITE_LOCKED is returned to the caller, the
 ** identity of the database connection (the blocking connection) that
-** has locked the required resource is stored internally. ^After an 
+** has locked the required resource is stored internally. ^After an
 ** application receives an SQLITE_LOCKED error, it may call the
-** sqlite3_unlock_notify() method with the blocked connection handle as 
+** sqlite3_unlock_notify() method with the blocked connection handle as
 ** the first argument to register for a callback that will be invoked
 ** when the blocking connections current transaction is concluded. ^The
 ** callback is invoked from within the [sqlite3_step] or [sqlite3_close]
@@ -8634,15 +9412,15 @@ SQLITE_API int sqlite3_backup_pagecount(sqlite3_backup *p);
 **
 ** ^If the blocked connection is attempting to obtain a write-lock on a
 ** shared-cache table, and more than one other connection currently holds
-** a read-lock on the same table, then SQLite arbitrarily selects one of 
+** a read-lock on the same table, then SQLite arbitrarily selects one of
 ** the other connections to use as the blocking connection.
 **
-** ^(There may be at most one unlock-notify callback registered by a 
+** ^(There may be at most one unlock-notify callback registered by a
 ** blocked connection. If sqlite3_unlock_notify() is called when the
 ** blocked connection already has a registered unlock-notify callback,
 ** then the new callback replaces the old.)^ ^If sqlite3_unlock_notify() is
 ** called with a NULL pointer as its second argument, then any existing
-** unlock-notify callback is canceled. ^The blocked connections 
+** unlock-notify callback is canceled. ^The blocked connections
 ** unlock-notify callback may also be canceled by closing the blocked
 ** connection using [sqlite3_close()].
 **
@@ -8655,7 +9433,7 @@ SQLITE_API int sqlite3_backup_pagecount(sqlite3_backup *p);
 **
 ** <b>Callback Invocation Details</b>
 **
-** When an unlock-notify callback is registered, the application provides a 
+** When an unlock-notify callback is registered, the application provides a
 ** single void* pointer that is passed to the callback when it is invoked.
 ** However, the signature of the callback function allows SQLite to pass
 ** it an array of void* context pointers. The first argument passed to
@@ -8668,12 +9446,12 @@ SQLITE_API int sqlite3_backup_pagecount(sqlite3_backup *p);
 ** same callback function, then instead of invoking the callback function
 ** multiple times, it is invoked once with the set of void* context pointers
 ** specified by the blocked connections bundled together into an array.
-** This gives the application an opportunity to prioritize any actions 
+** This gives the application an opportunity to prioritize any actions
 ** related to the set of unblocked database connections.
 **
 ** <b>Deadlock Detection</b>
 **
-** Assuming that after registering for an unlock-notify callback a 
+** Assuming that after registering for an unlock-notify callback a
 ** database waits for the callback to be issued before taking any further
 ** action (a reasonable assumption), then using this API may cause the
 ** application to deadlock. For example, if connection X is waiting for
@@ -8696,7 +9474,7 @@ SQLITE_API int sqlite3_backup_pagecount(sqlite3_backup *p);
 **
 ** <b>The "DROP TABLE" Exception</b>
 **
-** When a call to [sqlite3_step()] returns SQLITE_LOCKED, it is almost 
+** When a call to [sqlite3_step()] returns SQLITE_LOCKED, it is almost
 ** always appropriate to call sqlite3_unlock_notify(). There is however,
 ** one exception. When executing a "DROP TABLE" or "DROP INDEX" statement,
 ** SQLite checks if there are any currently executing SELECT statements
@@ -8709,7 +9487,7 @@ SQLITE_API int sqlite3_backup_pagecount(sqlite3_backup *p);
 ** One way around this problem is to check the extended error code returned
 ** by an sqlite3_step() call. ^(If there is a blocking connection, then the
 ** extended error code is set to SQLITE_LOCKED_SHAREDCACHE. Otherwise, in
-** the special "DROP TABLE/INDEX" case, the extended error code is just 
+** the special "DROP TABLE/INDEX" case, the extended error code is just
 ** SQLITE_LOCKED.)^
 */
 SQLITE_API int sqlite3_unlock_notify(
@@ -8800,8 +9578,8 @@ SQLITE_API void sqlite3_log(int iErrCode, const char *zFormat, ...);
 ** ^The [sqlite3_wal_hook()] function is used to register a callback that
 ** is invoked each time data is committed to a database in wal mode.
 **
-** ^(The callback is invoked by SQLite after the commit has taken place and 
-** the associated write-lock on the database released)^, so the implementation 
+** ^(The callback is invoked by SQLite after the commit has taken place and
+** the associated write-lock on the database released)^, so the implementation
 ** may read, write or [checkpoint] the database as required.
 **
 ** ^The first parameter passed to the callback function when it is invoked
@@ -8820,15 +9598,16 @@ SQLITE_API void sqlite3_log(int iErrCode, const char *zFormat, ...);
 ** that does not correspond to any valid SQLite error code, the results
 ** are undefined.
 **
-** A single database handle may have at most a single write-ahead log callback 
+** A single database handle may have at most a single write-ahead log callback
 ** registered at one time. ^Calling [sqlite3_wal_hook()] replaces any
-** previously registered write-ahead log callback. ^Note that the
-** [sqlite3_wal_autocheckpoint()] interface and the
+** previously registered write-ahead log callback. ^The return value is
+** a copy of the third parameter from the previous call, if any, or 0.
+** ^Note that the [sqlite3_wal_autocheckpoint()] interface and the
 ** [wal_autocheckpoint pragma] both invoke [sqlite3_wal_hook()] and will
 ** overwrite any prior [sqlite3_wal_hook()] settings.
 */
 SQLITE_API void *sqlite3_wal_hook(
-  sqlite3*, 
+  sqlite3*,
   int(*)(void *,sqlite3*,const char*,int),
   void*
 );
@@ -8841,7 +9620,7 @@ SQLITE_API void *sqlite3_wal_hook(
 ** [sqlite3_wal_hook()] that causes any database on [database connection] D
 ** to automatically [checkpoint]
 ** after committing a transaction if there are N or
-** more frames in the [write-ahead log] file.  ^Passing zero or 
+** more frames in the [write-ahead log] file.  ^Passing zero or
 ** a negative value as the nFrame parameter disables automatic
 ** checkpoints entirely.
 **
@@ -8871,7 +9650,7 @@ SQLITE_API int sqlite3_wal_autocheckpoint(sqlite3 *db, int N);
 ** ^(The sqlite3_wal_checkpoint(D,X) is equivalent to
 ** [sqlite3_wal_checkpoint_v2](D,X,[SQLITE_CHECKPOINT_PASSIVE],0,0).)^
 **
-** In brief, sqlite3_wal_checkpoint(D,X) causes the content in the 
+** In brief, sqlite3_wal_checkpoint(D,X) causes the content in the
 ** [write-ahead log] for database X on [database connection] D to be
 ** transferred into the database file and for the write-ahead log to
 ** be reset.  See the [checkpointing] documentation for addition
@@ -8897,10 +9676,10 @@ SQLITE_API int sqlite3_wal_checkpoint(sqlite3 *db, const char *zDb);
 **
 ** <dl>
 ** <dt>SQLITE_CHECKPOINT_PASSIVE<dd>
-**   ^Checkpoint as many frames as possible without waiting for any database 
-**   readers or writers to finish, then sync the database file if all frames 
+**   ^Checkpoint as many frames as possible without waiting for any database
+**   readers or writers to finish, then sync the database file if all frames
 **   in the log were checkpointed. ^The [busy-handler callback]
-**   is never invoked in the SQLITE_CHECKPOINT_PASSIVE mode.  
+**   is never invoked in the SQLITE_CHECKPOINT_PASSIVE mode.
 **   ^On the other hand, passive mode might leave the checkpoint unfinished
 **   if there are concurrent readers or writers.
 **
@@ -8914,9 +9693,9 @@ SQLITE_API int sqlite3_wal_checkpoint(sqlite3 *db, const char *zDb);
 **
 ** <dt>SQLITE_CHECKPOINT_RESTART<dd>
 **   ^This mode works the same way as SQLITE_CHECKPOINT_FULL with the addition
-**   that after checkpointing the log file it blocks (calls the 
+**   that after checkpointing the log file it blocks (calls the
 **   [busy-handler callback])
-**   until all readers are reading from the database file only. ^This ensures 
+**   until all readers are reading from the database file only. ^This ensures
 **   that the next writer will restart the log file from the beginning.
 **   ^Like SQLITE_CHECKPOINT_FULL, this mode blocks new
 **   database writer attempts while it is pending, but does not impede readers.
@@ -8938,31 +9717,31 @@ SQLITE_API int sqlite3_wal_checkpoint(sqlite3 *db, const char *zDb);
 ** truncated to zero bytes and so both *pnLog and *pnCkpt will be set to zero.
 **
 ** ^All calls obtain an exclusive "checkpoint" lock on the database file. ^If
-** any other process is running a checkpoint operation at the same time, the 
-** lock cannot be obtained and SQLITE_BUSY is returned. ^Even if there is a 
+** any other process is running a checkpoint operation at the same time, the
+** lock cannot be obtained and SQLITE_BUSY is returned. ^Even if there is a
 ** busy-handler configured, it will not be invoked in this case.
 **
-** ^The SQLITE_CHECKPOINT_FULL, RESTART and TRUNCATE modes also obtain the 
+** ^The SQLITE_CHECKPOINT_FULL, RESTART and TRUNCATE modes also obtain the
 ** exclusive "writer" lock on the database file. ^If the writer lock cannot be
 ** obtained immediately, and a busy-handler is configured, it is invoked and
 ** the writer lock retried until either the busy-handler returns 0 or the lock
 ** is successfully obtained. ^The busy-handler is also invoked while waiting for
 ** database readers as described above. ^If the busy-handler returns 0 before
 ** the writer lock is obtained or while waiting for database readers, the
-** checkpoint operation proceeds from that point in the same way as 
-** SQLITE_CHECKPOINT_PASSIVE - checkpointing as many frames as possible 
+** checkpoint operation proceeds from that point in the same way as
+** SQLITE_CHECKPOINT_PASSIVE - checkpointing as many frames as possible
 ** without blocking any further. ^SQLITE_BUSY is returned in this case.
 **
 ** ^If parameter zDb is NULL or points to a zero length string, then the
-** specified operation is attempted on all WAL databases [attached] to 
+** specified operation is attempted on all WAL databases [attached] to
 ** [database connection] db.  In this case the
-** values written to output parameters *pnLog and *pnCkpt are undefined. ^If 
-** an SQLITE_BUSY error is encountered when processing one or more of the 
-** attached WAL databases, the operation is still attempted on any remaining 
-** attached databases and SQLITE_BUSY is returned at the end. ^If any other 
-** error occurs while processing an attached database, processing is abandoned 
-** and the error code is returned to the caller immediately. ^If no error 
-** (SQLITE_BUSY or otherwise) is encountered while processing the attached 
+** values written to output parameters *pnLog and *pnCkpt are undefined. ^If
+** an SQLITE_BUSY error is encountered when processing one or more of the
+** attached WAL databases, the operation is still attempted on any remaining
+** attached databases and SQLITE_BUSY is returned at the end. ^If any other
+** error occurs while processing an attached database, processing is abandoned
+** and the error code is returned to the caller immediately. ^If no error
+** (SQLITE_BUSY or otherwise) is encountered while processing the attached
 ** databases, SQLITE_OK is returned.
 **
 ** ^If database zDb is the name of an attached database that is not in WAL
@@ -8997,7 +9776,7 @@ SQLITE_API int sqlite3_wal_checkpoint_v2(
 */
 #define SQLITE_CHECKPOINT_PASSIVE  0  /* Do as much as possible w/o blocking */
 #define SQLITE_CHECKPOINT_FULL     1  /* Wait for writers, then checkpoint */
-#define SQLITE_CHECKPOINT_RESTART  2  /* Like FULL but wait for for readers */
+#define SQLITE_CHECKPOINT_RESTART  2  /* Like FULL but wait for readers */
 #define SQLITE_CHECKPOINT_TRUNCATE 3  /* Like RESTART but also truncate WAL */
 
 /*
@@ -9022,7 +9801,7 @@ SQLITE_API int sqlite3_vtab_config(sqlite3*, int op, ...);
 
 /*
 ** CAPI3REF: Virtual Table Configuration Options
-** KEYWORDS: {virtual table configuration options} 
+** KEYWORDS: {virtual table configuration options}
 ** KEYWORDS: {virtual table configuration option}
 **
 ** These macros define the various options to the
@@ -9045,27 +9824,27 @@ SQLITE_API int sqlite3_vtab_config(sqlite3*, int op, ...);
 ** If X is non-zero, then the virtual table implementation guarantees
 ** that if [xUpdate] returns [SQLITE_CONSTRAINT], it will do so before
 ** any modifications to internal or persistent data structures have been made.
-** If the [ON CONFLICT] mode is ABORT, FAIL, IGNORE or ROLLBACK, SQLite 
+** If the [ON CONFLICT] mode is ABORT, FAIL, IGNORE or ROLLBACK, SQLite
 ** is able to roll back a statement or database transaction, and abandon
-** or continue processing the current SQL statement as appropriate. 
+** or continue processing the current SQL statement as appropriate.
 ** If the ON CONFLICT mode is REPLACE and the [xUpdate] method returns
 ** [SQLITE_CONSTRAINT], SQLite handles this as if the ON CONFLICT mode
 ** had been ABORT.
 **
 ** Virtual table implementations that are required to handle OR REPLACE
-** must do so within the [xUpdate] method. If a call to the 
-** [sqlite3_vtab_on_conflict()] function indicates that the current ON 
-** CONFLICT policy is REPLACE, the virtual table implementation should 
+** must do so within the [xUpdate] method. If a call to the
+** [sqlite3_vtab_on_conflict()] function indicates that the current ON
+** CONFLICT policy is REPLACE, the virtual table implementation should
 ** silently replace the appropriate rows within the xUpdate callback and
 ** return SQLITE_OK. Or, if this is not possible, it may return
-** SQLITE_CONSTRAINT, in which case SQLite falls back to OR ABORT 
+** SQLITE_CONSTRAINT, in which case SQLite falls back to OR ABORT
 ** constraint handling.
 ** </dd>
 **
 ** [[SQLITE_VTAB_DIRECTONLY]]<dt>SQLITE_VTAB_DIRECTONLY</dt>
 ** <dd>Calls of the form
 ** [sqlite3_vtab_config](db,SQLITE_VTAB_DIRECTONLY) from within the
-** the [xConnect] or [xCreate] methods of a [virtual table] implmentation
+** the [xConnect] or [xCreate] methods of a [virtual table] implementation
 ** prohibits that virtual table from being used from within triggers and
 ** views.
 ** </dd>
@@ -9073,18 +9852,28 @@ SQLITE_API int sqlite3_vtab_config(sqlite3*, int op, ...);
 ** [[SQLITE_VTAB_INNOCUOUS]]<dt>SQLITE_VTAB_INNOCUOUS</dt>
 ** <dd>Calls of the form
 ** [sqlite3_vtab_config](db,SQLITE_VTAB_INNOCUOUS) from within the
-** the [xConnect] or [xCreate] methods of a [virtual table] implmentation
+** the [xConnect] or [xCreate] methods of a [virtual table] implementation
 ** identify that virtual table as being safe to use from within triggers
 ** and views.  Conceptually, the SQLITE_VTAB_INNOCUOUS tag means that the
 ** virtual table can do no serious harm even if it is controlled by a
 ** malicious hacker.  Developers should avoid setting the SQLITE_VTAB_INNOCUOUS
 ** flag unless absolutely necessary.
 ** </dd>
+**
+** [[SQLITE_VTAB_USES_ALL_SCHEMAS]]<dt>SQLITE_VTAB_USES_ALL_SCHEMAS</dt>
+** <dd>Calls of the form
+** [sqlite3_vtab_config](db,SQLITE_VTAB_USES_ALL_SCHEMA) from within the
+** the [xConnect] or [xCreate] methods of a [virtual table] implementation
+** instruct the query planner to begin at least a read transaction on
+** all schemas ("main", "temp", and any ATTACH-ed databases) whenever the
+** virtual table is used.
+** </dd>
 ** </dl>
 */
 #define SQLITE_VTAB_CONSTRAINT_SUPPORT 1
 #define SQLITE_VTAB_INNOCUOUS          2
 #define SQLITE_VTAB_DIRECTONLY         3
+#define SQLITE_VTAB_USES_ALL_SCHEMAS   4
 
 /*
 ** CAPI3REF: Determine The Virtual Table Conflict Policy
@@ -9102,10 +9891,11 @@ SQLITE_API int sqlite3_vtab_on_conflict(sqlite3 *);
 ** CAPI3REF: Determine If Virtual Table Column Access Is For UPDATE
 **
 ** If the sqlite3_vtab_nochange(X) routine is called within the [xColumn]
-** method of a [virtual table], then it returns true if and only if the
+** method of a [virtual table], then it might return true if the
 ** column is being fetched as part of an UPDATE operation during which the
-** column value will not change.  Applications might use this to substitute
-** a return value that is less expensive to compute and that the corresponding
+** column value will not change.  The virtual table implementation can use
+** this hint as permission to substitute a return value that is less
+** expensive to compute and that the corresponding
 ** [xUpdate] method understands as a "no-change" value.
 **
 ** If the [xColumn] method calls sqlite3_vtab_nochange() and finds that
@@ -9114,23 +9904,306 @@ SQLITE_API int sqlite3_vtab_on_conflict(sqlite3 *);
 ** any of the [sqlite3_result_int|sqlite3_result_xxxxx() interfaces].
 ** In that case, [sqlite3_value_nochange(X)] will return true for the
 ** same column in the [xUpdate] method.
+**
+** The sqlite3_vtab_nochange() routine is an optimization.  Virtual table
+** implementations should continue to give a correct answer even if the
+** sqlite3_vtab_nochange() interface were to always return false.  In the
+** current implementation, the sqlite3_vtab_nochange() interface does always
+** returns false for the enhanced [UPDATE FROM] statement.
 */
 SQLITE_API int sqlite3_vtab_nochange(sqlite3_context*);
 
 /*
 ** CAPI3REF: Determine The Collation For a Virtual Table Constraint
+** METHOD: sqlite3_index_info
 **
 ** This function may only be called from within a call to the [xBestIndex]
-** method of a [virtual table]. 
+** method of a [virtual table].  This function returns a pointer to a string
+** that is the name of the appropriate collation sequence to use for text
+** comparisons on the constraint identified by its arguments.
+**
+** The first argument must be the pointer to the [sqlite3_index_info] object
+** that is the first parameter to the xBestIndex() method. The second argument
+** must be an index into the aConstraint[] array belonging to the
+** sqlite3_index_info structure passed to xBestIndex.
+**
+** Important:
+** The first parameter must be the same pointer that is passed into the
+** xBestMethod() method.  The first parameter may not be a pointer to a
+** different [sqlite3_index_info] object, even an exact copy.
+**
+** The return value is computed as follows:
+**
+** <ol>
+** <li><p> If the constraint comes from a WHERE clause expression that contains
+**         a [COLLATE operator], then the name of the collation specified by
+**         that COLLATE operator is returned.
+** <li><p> If there is no COLLATE operator, but the column that is the subject
+**         of the constraint specifies an alternative collating sequence via
+**         a [COLLATE clause] on the column definition within the CREATE TABLE
+**         statement that was passed into [sqlite3_declare_vtab()], then the
+**         name of that alternative collating sequence is returned.
+** <li><p> Otherwise, "BINARY" is returned.
+** </ol>
+*/
+SQLITE_API const char *sqlite3_vtab_collation(sqlite3_index_info*,int);
+
+/*
+** CAPI3REF: Determine if a virtual table query is DISTINCT
+** METHOD: sqlite3_index_info
+**
+** This API may only be used from within an [xBestIndex|xBestIndex method]
+** of a [virtual table] implementation. The result of calling this
+** interface from outside of xBestIndex() is undefined and probably harmful.
+**
+** ^The sqlite3_vtab_distinct() interface returns an integer between 0 and
+** 3.  The integer returned by sqlite3_vtab_distinct()
+** gives the virtual table additional information about how the query
+** planner wants the output to be ordered. As long as the virtual table
+** can meet the ordering requirements of the query planner, it may set
+** the "orderByConsumed" flag.
+**
+** <ol><li value="0"><p>
+** ^If the sqlite3_vtab_distinct() interface returns 0, that means
+** that the query planner needs the virtual table to return all rows in the
+** sort order defined by the "nOrderBy" and "aOrderBy" fields of the
+** [sqlite3_index_info] object.  This is the default expectation.  If the
+** virtual table outputs all rows in sorted order, then it is always safe for
+** the xBestIndex method to set the "orderByConsumed" flag, regardless of
+** the return value from sqlite3_vtab_distinct().
+** <li value="1"><p>
+** ^(If the sqlite3_vtab_distinct() interface returns 1, that means
+** that the query planner does not need the rows to be returned in sorted order
+** as long as all rows with the same values in all columns identified by the
+** "aOrderBy" field are adjacent.)^  This mode is used when the query planner
+** is doing a GROUP BY.
+** <li value="2"><p>
+** ^(If the sqlite3_vtab_distinct() interface returns 2, that means
+** that the query planner does not need the rows returned in any particular
+** order, as long as rows with the same values in all columns identified
+** by "aOrderBy" are adjacent.)^  ^(Furthermore, when two or more rows
+** contain the same values for all columns identified by "colUsed", all but
+** one such row may optionally be omitted from the result.)^
+** The virtual table is not required to omit rows that are duplicates
+** over the "colUsed" columns, but if the virtual table can do that without
+** too much extra effort, it could potentially help the query to run faster.
+** This mode is used for a DISTINCT query.
+** <li value="3"><p>
+** ^(If the sqlite3_vtab_distinct() interface returns 3, that means the
+** virtual table must return rows in the order defined by "aOrderBy" as
+** if the sqlite3_vtab_distinct() interface had returned 0.  However if
+** two or more rows in the result have the same values for all columns
+** identified by "colUsed", then all but one such row may optionally be
+** omitted.)^  Like when the return value is 2, the virtual table
+** is not required to omit rows that are duplicates over the "colUsed"
+** columns, but if the virtual table can do that without
+** too much extra effort, it could potentially help the query to run faster.
+** This mode is used for queries
+** that have both DISTINCT and ORDER BY clauses.
+** </ol>
+**
+** <p>The following table summarizes the conditions under which the
+** virtual table is allowed to set the "orderByConsumed" flag based on
+** the value returned by sqlite3_vtab_distinct().  This table is a
+** restatement of the previous four paragraphs:
+**
+** <table border=1 cellspacing=0 cellpadding=10 width="90%">
+** <tr>
+** <td valign="top">sqlite3_vtab_distinct() return value
+** <td valign="top">Rows are returned in aOrderBy order
+** <td valign="top">Rows with the same value in all aOrderBy columns are adjacent
+** <td valign="top">Duplicates over all colUsed columns may be omitted
+** <tr><td>0<td>yes<td>yes<td>no
+** <tr><td>1<td>no<td>yes<td>no
+** <tr><td>2<td>no<td>yes<td>yes
+** <tr><td>3<td>yes<td>yes<td>yes
+** </table>
+**
+** ^For the purposes of comparing virtual table output values to see if the
+** values are same value for sorting purposes, two NULL values are considered
+** to be the same.  In other words, the comparison operator is "IS"
+** (or "IS NOT DISTINCT FROM") and not "==".
+**
+** If a virtual table implementation is unable to meet the requirements
+** specified above, then it must not set the "orderByConsumed" flag in the
+** [sqlite3_index_info] object or an incorrect answer may result.
+**
+** ^A virtual table implementation is always free to return rows in any order
+** it wants, as long as the "orderByConsumed" flag is not set.  ^When the
+** the "orderByConsumed" flag is unset, the query planner will add extra
+** [bytecode] to ensure that the final results returned by the SQL query are
+** ordered correctly.  The use of the "orderByConsumed" flag and the
+** sqlite3_vtab_distinct() interface is merely an optimization.  ^Careful
+** use of the sqlite3_vtab_distinct() interface and the "orderByConsumed"
+** flag might help queries against a virtual table to run faster.  Being
+** overly aggressive and setting the "orderByConsumed" flag when it is not
+** valid to do so, on the other hand, might cause SQLite to return incorrect
+** results.
+*/
+SQLITE_API int sqlite3_vtab_distinct(sqlite3_index_info*);
+
+/*
+** CAPI3REF: Identify and handle IN constraints in xBestIndex
+**
+** This interface may only be used from within an
+** [xBestIndex|xBestIndex() method] of a [virtual table] implementation.
+** The result of invoking this interface from any other context is
+** undefined and probably harmful.
+**
+** ^(A constraint on a virtual table of the form
+** "[IN operator|column IN (...)]" is
+** communicated to the xBestIndex method as a
+** [SQLITE_INDEX_CONSTRAINT_EQ] constraint.)^  If xBestIndex wants to use
+** this constraint, it must set the corresponding
+** aConstraintUsage[].argvIndex to a positive integer.  ^(Then, under
+** the usual mode of handling IN operators, SQLite generates [bytecode]
+** that invokes the [xFilter|xFilter() method] once for each value
+** on the right-hand side of the IN operator.)^  Thus the virtual table
+** only sees a single value from the right-hand side of the IN operator
+** at a time.
+**
+** In some cases, however, it would be advantageous for the virtual
+** table to see all values on the right-hand of the IN operator all at
+** once.  The sqlite3_vtab_in() interfaces facilitates this in two ways:
 **
-** The first argument must be the sqlite3_index_info object that is the
-** first parameter to the xBestIndex() method. The second argument must be
-** an index into the aConstraint[] array belonging to the sqlite3_index_info
-** structure passed to xBestIndex. This function returns a pointer to a buffer 
-** containing the name of the collation sequence for the corresponding
-** constraint.
+** <ol>
+** <li><p>
+**   ^A call to sqlite3_vtab_in(P,N,-1) will return true (non-zero)
+**   if and only if the [sqlite3_index_info|P->aConstraint][N] constraint
+**   is an [IN operator] that can be processed all at once.  ^In other words,
+**   sqlite3_vtab_in() with -1 in the third argument is a mechanism
+**   by which the virtual table can ask SQLite if all-at-once processing
+**   of the IN operator is even possible.
+**
+** <li><p>
+**   ^A call to sqlite3_vtab_in(P,N,F) with F==1 or F==0 indicates
+**   to SQLite that the virtual table does or does not want to process
+**   the IN operator all-at-once, respectively.  ^Thus when the third
+**   parameter (F) is non-negative, this interface is the mechanism by
+**   which the virtual table tells SQLite how it wants to process the
+**   IN operator.
+** </ol>
+**
+** ^The sqlite3_vtab_in(P,N,F) interface can be invoked multiple times
+** within the same xBestIndex method call.  ^For any given P,N pair,
+** the return value from sqlite3_vtab_in(P,N,F) will always be the same
+** within the same xBestIndex call.  ^If the interface returns true
+** (non-zero), that means that the constraint is an IN operator
+** that can be processed all-at-once.  ^If the constraint is not an IN
+** operator or cannot be processed all-at-once, then the interface returns
+** false.
+**
+** ^(All-at-once processing of the IN operator is selected if both of the
+** following conditions are met:
+**
+** <ol>
+** <li><p> The P->aConstraintUsage[N].argvIndex value is set to a positive
+** integer.  This is how the virtual table tells SQLite that it wants to
+** use the N-th constraint.
+**
+** <li><p> The last call to sqlite3_vtab_in(P,N,F) for which F was
+** non-negative had F>=1.
+** </ol>)^
+**
+** ^If either or both of the conditions above are false, then SQLite uses
+** the traditional one-at-a-time processing strategy for the IN constraint.
+** ^If both conditions are true, then the argvIndex-th parameter to the
+** xFilter method will be an [sqlite3_value] that appears to be NULL,
+** but which can be passed to [sqlite3_vtab_in_first()] and
+** [sqlite3_vtab_in_next()] to find all values on the right-hand side
+** of the IN constraint.
 */
-SQLITE_API SQLITE_EXPERIMENTAL const char *sqlite3_vtab_collation(sqlite3_index_info*,int);
+SQLITE_API int sqlite3_vtab_in(sqlite3_index_info*, int iCons, int bHandle);
+
+/*
+** CAPI3REF: Find all elements on the right-hand side of an IN constraint.
+**
+** These interfaces are only useful from within the
+** [xFilter|xFilter() method] of a [virtual table] implementation.
+** The result of invoking these interfaces from any other context
+** is undefined and probably harmful.
+**
+** The X parameter in a call to sqlite3_vtab_in_first(X,P) or
+** sqlite3_vtab_in_next(X,P) should be one of the parameters to the
+** xFilter method which invokes these routines, and specifically
+** a parameter that was previously selected for all-at-once IN constraint
+** processing use the [sqlite3_vtab_in()] interface in the
+** [xBestIndex|xBestIndex method].  ^(If the X parameter is not
+** an xFilter argument that was selected for all-at-once IN constraint
+** processing, then these routines return [SQLITE_ERROR].)^
+**
+** ^(Use these routines to access all values on the right-hand side
+** of the IN constraint using code like the following:
+**
+** <blockquote><pre>
+** &nbsp;  for(rc=sqlite3_vtab_in_first(pList, &pVal);
+** &nbsp;      rc==SQLITE_OK && pVal;
+** &nbsp;      rc=sqlite3_vtab_in_next(pList, &pVal)
+** &nbsp;  ){
+** &nbsp;    // do something with pVal
+** &nbsp;  }
+** &nbsp;  if( rc!=SQLITE_OK ){
+** &nbsp;    // an error has occurred
+** &nbsp;  }
+** </pre></blockquote>)^
+**
+** ^On success, the sqlite3_vtab_in_first(X,P) and sqlite3_vtab_in_next(X,P)
+** routines return SQLITE_OK and set *P to point to the first or next value
+** on the RHS of the IN constraint.  ^If there are no more values on the
+** right hand side of the IN constraint, then *P is set to NULL and these
+** routines return [SQLITE_DONE].  ^The return value might be
+** some other value, such as SQLITE_NOMEM, in the event of a malfunction.
+**
+** The *ppOut values returned by these routines are only valid until the
+** next call to either of these routines or until the end of the xFilter
+** method from which these routines were called.  If the virtual table
+** implementation needs to retain the *ppOut values for longer, it must make
+** copies.  The *ppOut values are [protected sqlite3_value|protected].
+*/
+SQLITE_API int sqlite3_vtab_in_first(sqlite3_value *pVal, sqlite3_value **ppOut);
+SQLITE_API int sqlite3_vtab_in_next(sqlite3_value *pVal, sqlite3_value **ppOut);
+
+/*
+** CAPI3REF: Constraint values in xBestIndex()
+** METHOD: sqlite3_index_info
+**
+** This API may only be used from within the [xBestIndex|xBestIndex method]
+** of a [virtual table] implementation. The result of calling this interface
+** from outside of an xBestIndex method are undefined and probably harmful.
+**
+** ^When the sqlite3_vtab_rhs_value(P,J,V) interface is invoked from within
+** the [xBestIndex] method of a [virtual table] implementation, with P being
+** a copy of the [sqlite3_index_info] object pointer passed into xBestIndex and
+** J being a 0-based index into P->aConstraint[], then this routine
+** attempts to set *V to the value of the right-hand operand of
+** that constraint if the right-hand operand is known.  ^If the
+** right-hand operand is not known, then *V is set to a NULL pointer.
+** ^The sqlite3_vtab_rhs_value(P,J,V) interface returns SQLITE_OK if
+** and only if *V is set to a value.  ^The sqlite3_vtab_rhs_value(P,J,V)
+** inteface returns SQLITE_NOTFOUND if the right-hand side of the J-th
+** constraint is not available.  ^The sqlite3_vtab_rhs_value() interface
+** can return an result code other than SQLITE_OK or SQLITE_NOTFOUND if
+** something goes wrong.
+**
+** The sqlite3_vtab_rhs_value() interface is usually only successful if
+** the right-hand operand of a constraint is a literal value in the original
+** SQL statement.  If the right-hand operand is an expression or a reference
+** to some other column or a [host parameter], then sqlite3_vtab_rhs_value()
+** will probably return [SQLITE_NOTFOUND].
+**
+** ^(Some constraints, such as [SQLITE_INDEX_CONSTRAINT_ISNULL] and
+** [SQLITE_INDEX_CONSTRAINT_ISNOTNULL], have no right-hand operand.  For such
+** constraints, sqlite3_vtab_rhs_value() always returns SQLITE_NOTFOUND.)^
+**
+** ^The [sqlite3_value] object returned in *V is a protected sqlite3_value
+** and remains valid for the duration of the xBestIndex method call.
+** ^When xBestIndex returns, the sqlite3_value object returned by
+** sqlite3_vtab_rhs_value() is automatically deallocated.
+**
+** The "_rhs_" in the name of this routine is an abbreviation for
+** "Right-Hand Side".
+*/
+SQLITE_API int sqlite3_vtab_rhs_value(sqlite3_index_info*, int, sqlite3_value **ppVal);
 
 /*
 ** CAPI3REF: Conflict resolution modes
@@ -9162,6 +10235,10 @@ SQLITE_API SQLITE_EXPERIMENTAL const char *sqlite3_vtab_collation(sqlite3_index_
 ** managed by the prepared statement S and will be automatically freed when
 ** S is finalized.
 **
+** Not all values are available for all query elements. When a value is
+** not available, the output variable is set to -1 if the value is numeric,
+** or to NULL if it is a string (SQLITE_SCANSTAT_NAME).
+**
 ** <dl>
 ** [[SQLITE_SCANSTAT_NLOOP]] <dt>SQLITE_SCANSTAT_NLOOP</dt>
 ** <dd>^The [sqlite3_int64] variable pointed to by the V parameter will be
@@ -9189,12 +10266,24 @@ SQLITE_API SQLITE_EXPERIMENTAL const char *sqlite3_vtab_collation(sqlite3_index_
 ** to a zero-terminated UTF-8 string containing the [EXPLAIN QUERY PLAN]
 ** description for the X-th loop.
 **
-** [[SQLITE_SCANSTAT_SELECTID]] <dt>SQLITE_SCANSTAT_SELECT</dt>
+** [[SQLITE_SCANSTAT_SELECTID]] <dt>SQLITE_SCANSTAT_SELECTID</dt>
 ** <dd>^The "int" variable pointed to by the V parameter will be set to the
-** "select-id" for the X-th loop.  The select-id identifies which query or
-** subquery the loop is part of.  The main query has a select-id of zero.
-** The select-id is the same value as is output in the first column
-** of an [EXPLAIN QUERY PLAN] query.
+** id for the X-th query plan element. The id value is unique within the
+** statement. The select-id is the same value as is output in the first
+** column of an [EXPLAIN QUERY PLAN] query.
+**
+** [[SQLITE_SCANSTAT_PARENTID]] <dt>SQLITE_SCANSTAT_PARENTID</dt>
+** <dd>The "int" variable pointed to by the V parameter will be set to the
+** the id of the parent of the current query element, if applicable, or
+** to zero if the query element has no parent. This is the same value as
+** returned in the second column of an [EXPLAIN QUERY PLAN] query.
+**
+** [[SQLITE_SCANSTAT_NCYCLE]] <dt>SQLITE_SCANSTAT_NCYCLE</dt>
+** <dd>The sqlite3_int64 output value is set to the number of cycles,
+** according to the processor time-stamp counter, that elapsed while the
+** query element was being processed. This value is not available for
+** all query elements - if it is unavailable the output variable is
+** set to -1.
 ** </dl>
 */
 #define SQLITE_SCANSTAT_NLOOP    0
@@ -9203,12 +10292,14 @@ SQLITE_API SQLITE_EXPERIMENTAL const char *sqlite3_vtab_collation(sqlite3_index_
 #define SQLITE_SCANSTAT_NAME     3
 #define SQLITE_SCANSTAT_EXPLAIN  4
 #define SQLITE_SCANSTAT_SELECTID 5
+#define SQLITE_SCANSTAT_PARENTID 6
+#define SQLITE_SCANSTAT_NCYCLE   7
 
 /*
 ** CAPI3REF: Prepared Statement Scan Status
 ** METHOD: sqlite3_stmt
 **
-** This interface returns information about the predicted and measured
+** These interfaces return information about the predicted and measured
 ** performance for pStmt.  Advanced applications can use this
 ** interface to compare the predicted and the measured performance and
 ** issue warnings and/or rerun [ANALYZE] if discrepancies are found.
@@ -9219,19 +10310,25 @@ SQLITE_API SQLITE_EXPERIMENTAL const char *sqlite3_vtab_collation(sqlite3_index_
 **
 ** The "iScanStatusOp" parameter determines which status information to return.
 ** The "iScanStatusOp" must be one of the [scanstatus options] or the behavior
-** of this interface is undefined.
-** ^The requested measurement is written into a variable pointed to by
-** the "pOut" parameter.
-** Parameter "idx" identifies the specific loop to retrieve statistics for.
-** Loops are numbered starting from zero. ^If idx is out of range - less than
-** zero or greater than or equal to the total number of loops used to implement
-** the statement - a non-zero value is returned and the variable that pOut
-** points to is unchanged.
-**
-** ^Statistics might not be available for all loops in all statements. ^In cases
-** where there exist loops with no available statistics, this function behaves
-** as if the loop did not exist - it returns non-zero and leave the variable
-** that pOut points to unchanged.
+** of this interface is undefined. ^The requested measurement is written into
+** a variable pointed to by the "pOut" parameter.
+**
+** The "flags" parameter must be passed a mask of flags. At present only
+** one flag is defined - SQLITE_SCANSTAT_COMPLEX. If SQLITE_SCANSTAT_COMPLEX
+** is specified, then status information is available for all elements
+** of a query plan that are reported by "EXPLAIN QUERY PLAN" output. If
+** SQLITE_SCANSTAT_COMPLEX is not specified, then only query plan elements
+** that correspond to query loops (the "SCAN..." and "SEARCH..." elements of
+** the EXPLAIN QUERY PLAN output) are available. Invoking API
+** sqlite3_stmt_scanstatus() is equivalent to calling
+** sqlite3_stmt_scanstatus_v2() with a zeroed flags parameter.
+**
+** Parameter "idx" identifies the specific query element to retrieve statistics
+** for. Query elements are numbered starting from zero. A value of -1 may be
+** to query for statistics regarding the entire query. ^If idx is out of range
+** - less than -1 or greater than or equal to the total number of query
+** elements used to implement the statement - a non-zero value is returned and
+** the variable that pOut points to is unchanged.
 **
 ** See also: [sqlite3_stmt_scanstatus_reset()]
 */
@@ -9240,7 +10337,20 @@ SQLITE_API int sqlite3_stmt_scanstatus(
   int idx,                  /* Index of loop to report on */
   int iScanStatusOp,        /* Information desired.  SQLITE_SCANSTAT_* */
   void *pOut                /* Result written here */
-);     
+);
+SQLITE_API int sqlite3_stmt_scanstatus_v2(
+  sqlite3_stmt *pStmt,      /* Prepared statement for which info desired */
+  int idx,                  /* Index of loop to report on */
+  int iScanStatusOp,        /* Information desired.  SQLITE_SCANSTAT_* */
+  int flags,                /* Mask of flags defined below */
+  void *pOut                /* Result written here */
+);
+
+/*
+** CAPI3REF: Prepared Statement Scan Status
+** KEYWORDS: {scan status flags}
+*/
+#define SQLITE_SCANSTAT_COMPLEX 0x0001
 
 /*
 ** CAPI3REF: Zero Scan-Status Counters
@@ -9255,18 +10365,19 @@ SQLITE_API void sqlite3_stmt_scanstatus_reset(sqlite3_stmt*);
 
 /*
 ** CAPI3REF: Flush caches to disk mid-transaction
+** METHOD: sqlite3
 **
 ** ^If a write-transaction is open on [database connection] D when the
 ** [sqlite3_db_cacheflush(D)] interface invoked, any dirty
-** pages in the pager-cache that are not currently in use are written out 
+** pages in the pager-cache that are not currently in use are written out
 ** to disk. A dirty page may be in use if a database cursor created by an
 ** active SQL statement is reading from it, or if it is page 1 of a database
 ** file (page 1 is always "in use").  ^The [sqlite3_db_cacheflush(D)]
 ** interface flushes caches for all schemas - "main", "temp", and
 ** any [attached] databases.
 **
-** ^If this function needs to obtain extra database locks before dirty pages 
-** can be flushed to disk, it does so. ^If those locks cannot be obtained 
+** ^If this function needs to obtain extra database locks before dirty pages
+** can be flushed to disk, it does so. ^If those locks cannot be obtained
 ** immediately and there is a busy-handler callback configured, it is invoked
 ** in the usual manner. ^If the required lock still cannot be obtained, then
 ** the database is skipped and an attempt made to flush any dirty pages
@@ -9287,6 +10398,7 @@ SQLITE_API int sqlite3_db_cacheflush(sqlite3*);
 
 /*
 ** CAPI3REF: The pre-update hook.
+** METHOD: sqlite3
 **
 ** ^These interfaces are only available if SQLite is compiled using the
 ** [SQLITE_ENABLE_PREUPDATE_HOOK] compile-time option.
@@ -9304,7 +10416,7 @@ SQLITE_API int sqlite3_db_cacheflush(sqlite3*);
 **
 ** ^The preupdate hook only fires for changes to real database tables; the
 ** preupdate hook is not invoked for changes to [virtual tables] or to
-** system tables like sqlite_master or sqlite_stat1.
+** system tables like sqlite_sequence or sqlite_stat1.
 **
 ** ^The second parameter to the preupdate callback is a pointer to
 ** the [database connection] that registered the preupdate hook.
@@ -9313,21 +10425,25 @@ SQLITE_API int sqlite3_db_cacheflush(sqlite3*);
 ** kind of update operation that is about to occur.
 ** ^(The fourth parameter to the preupdate callback is the name of the
 ** database within the database connection that is being modified.  This
-** will be "main" for the main database or "temp" for TEMP tables or 
+** will be "main" for the main database or "temp" for TEMP tables or
 ** the name given after the AS keyword in the [ATTACH] statement for attached
 ** databases.)^
 ** ^The fifth parameter to the preupdate callback is the name of the
 ** table that is being modified.
 **
 ** For an UPDATE or DELETE operation on a [rowid table], the sixth
-** parameter passed to the preupdate callback is the initial [rowid] of the 
+** parameter passed to the preupdate callback is the initial [rowid] of the
 ** row being modified or deleted. For an INSERT operation on a rowid table,
-** or any operation on a WITHOUT ROWID table, the value of the sixth 
+** or any operation on a WITHOUT ROWID table, the value of the sixth
 ** parameter is undefined. For an INSERT or UPDATE on a rowid table the
 ** seventh parameter is the final rowid value of the row being inserted
 ** or updated. The value of the seventh parameter passed to the callback
 ** function is not defined for operations on WITHOUT ROWID tables, or for
-** INSERT operations on rowid tables.
+** DELETE operations on rowid tables.
+**
+** ^The sqlite3_preupdate_hook(D,C,P) function returns the P argument from
+** the previous call on the same [database connection] D, or NULL for
+** the first call on D.
 **
 ** The [sqlite3_preupdate_old()], [sqlite3_preupdate_new()],
 ** [sqlite3_preupdate_count()], and [sqlite3_preupdate_depth()] interfaces
@@ -9361,10 +10477,19 @@ SQLITE_API int sqlite3_db_cacheflush(sqlite3*);
 **
 ** ^The [sqlite3_preupdate_depth(D)] interface returns 0 if the preupdate
 ** callback was invoked as a result of a direct insert, update, or delete
-** operation; or 1 for inserts, updates, or deletes invoked by top-level 
+** operation; or 1 for inserts, updates, or deletes invoked by top-level
 ** triggers; or 2 for changes resulting from triggers called by top-level
 ** triggers; and so forth.
 **
+** When the [sqlite3_blob_write()] API is used to update a blob column,
+** the pre-update hook is invoked with SQLITE_DELETE. This is because the
+** in this case the new values are not available. In this case, when a
+** callback made with op==SQLITE_DELETE is actually a write using the
+** sqlite3_blob_write() API, the [sqlite3_preupdate_blobwrite()] returns
+** the index of the column being written. In other cases, where the
+** pre-update hook is being invoked for some other reason, including a
+** regular DELETE, sqlite3_preupdate_blobwrite() returns -1.
+**
 ** See also:  [sqlite3_update_hook()]
 */
 #if defined(SQLITE_ENABLE_PREUPDATE_HOOK)
@@ -9385,17 +10510,19 @@ SQLITE_API int sqlite3_preupdate_old(sqlite3 *, int, sqlite3_value **);
 SQLITE_API int sqlite3_preupdate_count(sqlite3 *);
 SQLITE_API int sqlite3_preupdate_depth(sqlite3 *);
 SQLITE_API int sqlite3_preupdate_new(sqlite3 *, int, sqlite3_value **);
+SQLITE_API int sqlite3_preupdate_blobwrite(sqlite3 *);
 #endif
 
 /*
 ** CAPI3REF: Low-level system error code
+** METHOD: sqlite3
 **
 ** ^Attempt to return the underlying operating system error code or error
 ** number that caused the most recent I/O error or failure to open a file.
 ** The return value is OS-dependent.  For example, on unix systems, after
 ** [sqlite3_open_v2()] returns [SQLITE_CANTOPEN], this interface could be
 ** called to get back the underlying "errno" that caused the problem, such
-** as ENOSPC, EAUTH, EISDIR, and so forth.  
+** as ENOSPC, EAUTH, EISDIR, and so forth.
 */
 SQLITE_API int sqlite3_system_errno(sqlite3*);
 
@@ -9433,12 +10560,20 @@ typedef struct sqlite3_snapshot {
 ** [sqlite3_snapshot_get(D,S,P)] interface writes a pointer to the newly
 ** created [sqlite3_snapshot] object into *P and returns SQLITE_OK.
 ** If there is not already a read-transaction open on schema S when
-** this function is called, one is opened automatically. 
+** this function is called, one is opened automatically.
+**
+** If a read-transaction is opened by this function, then it is guaranteed
+** that the returned snapshot object may not be invalidated by a database
+** writer or checkpointer until after the read-transaction is closed. This
+** is not guaranteed if a read-transaction is already open when this
+** function is called. In that case, any subsequent write or checkpoint
+** operation on the database may invalidate the returned snapshot handle,
+** even while the read-transaction remains open.
 **
 ** The following must be true for this function to succeed. If any of
 ** the following statements are false when sqlite3_snapshot_get() is
 ** called, SQLITE_ERROR is returned. The final value of *P is undefined
-** in this case. 
+** in this case.
 **
 ** <ul>
 **   <li> The database handle must not be in [autocommit mode].
@@ -9450,13 +10585,13 @@ typedef struct sqlite3_snapshot {
 **
 **   <li> One or more transactions must have been written to the current wal
 **        file since it was created on disk (by any connection). This means
-**        that a snapshot cannot be taken on a wal mode database with no wal 
+**        that a snapshot cannot be taken on a wal mode database with no wal
 **        file immediately after it is first opened. At least one transaction
 **        must be written to it first.
 ** </ul>
 **
 ** This function may also return SQLITE_NOMEM.  If it is called with the
-** database handle in autocommit mode but fails for some other reason, 
+** database handle in autocommit mode but fails for some other reason,
 ** whether or not a read transaction is opened on schema S is undefined.
 **
 ** The [sqlite3_snapshot] object returned from a successful call to
@@ -9476,38 +10611,38 @@ SQLITE_API SQLITE_EXPERIMENTAL int sqlite3_snapshot_get(
 ** CAPI3REF: Start a read transaction on an historical snapshot
 ** METHOD: sqlite3_snapshot
 **
-** ^The [sqlite3_snapshot_open(D,S,P)] interface either starts a new read 
-** transaction or upgrades an existing one for schema S of 
-** [database connection] D such that the read transaction refers to 
-** historical [snapshot] P, rather than the most recent change to the 
-** database. ^The [sqlite3_snapshot_open()] interface returns SQLITE_OK 
+** ^The [sqlite3_snapshot_open(D,S,P)] interface either starts a new read
+** transaction or upgrades an existing one for schema S of
+** [database connection] D such that the read transaction refers to
+** historical [snapshot] P, rather than the most recent change to the
+** database. ^The [sqlite3_snapshot_open()] interface returns SQLITE_OK
 ** on success or an appropriate [error code] if it fails.
 **
-** ^In order to succeed, the database connection must not be in 
+** ^In order to succeed, the database connection must not be in
 ** [autocommit mode] when [sqlite3_snapshot_open(D,S,P)] is called. If there
 ** is already a read transaction open on schema S, then the database handle
 ** must have no active statements (SELECT statements that have been passed
-** to sqlite3_step() but not sqlite3_reset() or sqlite3_finalize()). 
+** to sqlite3_step() but not sqlite3_reset() or sqlite3_finalize()).
 ** SQLITE_ERROR is returned if either of these conditions is violated, or
 ** if schema S does not exist, or if the snapshot object is invalid.
 **
 ** ^A call to sqlite3_snapshot_open() will fail to open if the specified
-** snapshot has been overwritten by a [checkpoint]. In this case 
+** snapshot has been overwritten by a [checkpoint]. In this case
 ** SQLITE_ERROR_SNAPSHOT is returned.
 **
-** If there is already a read transaction open when this function is 
+** If there is already a read transaction open when this function is
 ** invoked, then the same read transaction remains open (on the same
 ** database snapshot) if SQLITE_ERROR, SQLITE_BUSY or SQLITE_ERROR_SNAPSHOT
 ** is returned. If another error code - for example SQLITE_PROTOCOL or an
 ** SQLITE_IOERR error code - is returned, then the final state of the
-** read transaction is undefined. If SQLITE_OK is returned, then the 
+** read transaction is undefined. If SQLITE_OK is returned, then the
 ** read transaction is now open on database snapshot P.
 **
 ** ^(A call to [sqlite3_snapshot_open(D,S,P)] will fail if the
 ** database connection D does not know that the database file for
 ** schema S is in [WAL mode].  A database connection might not know
 ** that the database file is in [WAL mode] if there has been no prior
-** I/O on that database connection, or if the database entered [WAL mode] 
+** I/O on that database connection, or if the database entered [WAL mode]
 ** after the most recent I/O on the database connection.)^
 ** (Hint: Run "[PRAGMA application_id]" against a newly opened
 ** database connection in order to make it ready to use snapshots.)
@@ -9539,17 +10674,17 @@ SQLITE_API SQLITE_EXPERIMENTAL void sqlite3_snapshot_free(sqlite3_snapshot*);
 ** METHOD: sqlite3_snapshot
 **
 ** The sqlite3_snapshot_cmp(P1, P2) interface is used to compare the ages
-** of two valid snapshot handles. 
+** of two valid snapshot handles.
 **
-** If the two snapshot handles are not associated with the same database 
-** file, the result of the comparison is undefined. 
+** If the two snapshot handles are not associated with the same database
+** file, the result of the comparison is undefined.
 **
 ** Additionally, the result of the comparison is only valid if both of the
 ** snapshot handles were obtained by calling sqlite3_snapshot_get() since the
 ** last time the wal file was deleted. The wal file is deleted when the
 ** database is changed back to rollback mode or when the number of database
-** clients drops to zero. If either snapshot handle was obtained before the 
-** wal file was last deleted, the value returned by this function 
+** clients drops to zero. If either snapshot handle was obtained before the
+** wal file was last deleted, the value returned by this function
 ** is undefined.
 **
 ** Otherwise, this API returns a negative value if P1 refers to an older
@@ -9614,16 +10749,23 @@ SQLITE_API SQLITE_EXPERIMENTAL int sqlite3_snapshot_recover(sqlite3 *db, const c
 ** representation of the database will usually only exist if there has
 ** been a prior call to [sqlite3_deserialize(D,S,...)] with the same
 ** values of D and S.
-** The size of the database is written into *P even if the 
+** The size of the database is written into *P even if the
 ** SQLITE_SERIALIZE_NOCOPY bit is set but no contiguous copy
 ** of the database exists.
 **
+** After the call, if the SQLITE_SERIALIZE_NOCOPY bit had been set,
+** the returned buffer content will remain accessible and unchanged
+** until either the next write operation on the connection or when
+** the connection is closed, and applications must not modify the
+** buffer. If the bit had been clear, the returned buffer will not
+** be accessed by SQLite after the call.
+**
 ** A call to sqlite3_serialize(D,S,P,F) might return NULL even if the
 ** SQLITE_SERIALIZE_NOCOPY bit is omitted from argument F if a memory
 ** allocation error occurs.
 **
-** This interface is only available if SQLite is compiled with the
-** [SQLITE_ENABLE_DESERIALIZE] option.
+** This interface is omitted if SQLite is compiled with the
+** [SQLITE_OMIT_DESERIALIZE] option.
 */
 SQLITE_API unsigned char *sqlite3_serialize(
   sqlite3 *db,           /* The database connection */
@@ -9651,7 +10793,7 @@ SQLITE_API unsigned char *sqlite3_serialize(
 /*
 ** CAPI3REF: Deserialize a database
 **
-** The sqlite3_deserialize(D,S,P,N,M,F) interface causes the 
+** The sqlite3_deserialize(D,S,P,N,M,F) interface causes the
 ** [database connection] D to disconnect from database S and then
 ** reopen S as an in-memory database based on the serialization contained
 ** in P.  The serialized database P is N bytes in size.  M is the size of
@@ -9666,16 +10808,30 @@ SQLITE_API unsigned char *sqlite3_serialize(
 ** SQLite will try to increase the buffer size using sqlite3_realloc64()
 ** if writes on the database cause it to grow larger than M bytes.
 **
+** Applications must not modify the buffer P or invalidate it before
+** the database connection D is closed.
+**
 ** The sqlite3_deserialize() interface will fail with SQLITE_BUSY if the
 ** database is currently in a read transaction or is involved in a backup
 ** operation.
 **
-** If sqlite3_deserialize(D,S,P,N,M,F) fails for any reason and if the 
+** It is not possible to deserialized into the TEMP database.  If the
+** S argument to sqlite3_deserialize(D,S,P,N,M,F) is "temp" then the
+** function returns SQLITE_ERROR.
+**
+** The deserialized database should not be in [WAL mode].  If the database
+** is in WAL mode, then any attempt to use the database file will result
+** in an [SQLITE_CANTOPEN] error.  The application can set the
+** [file format version numbers] (bytes 18 and 19) of the input database P
+** to 0x01 prior to invoking sqlite3_deserialize(D,S,P,N,M,F) to force the
+** database file into rollback mode and work around this limitation.
+**
+** If sqlite3_deserialize(D,S,P,N,M,F) fails for any reason and if the
 ** SQLITE_DESERIALIZE_FREEONCLOSE bit is set in argument F, then
 ** [sqlite3_free()] is invoked on argument P prior to returning.
 **
-** This interface is only available if SQLite is compiled with the
-** [SQLITE_ENABLE_DESERIALIZE] option.
+** This interface is omitted if SQLite is compiled with the
+** [SQLITE_OMIT_DESERIALIZE] option.
 */
 SQLITE_API int sqlite3_deserialize(
   sqlite3 *db,            /* The database connection */
@@ -9719,6 +10875,17 @@ SQLITE_API int sqlite3_deserialize(
 # undef double
 #endif
 
+#if defined(__wasi__)
+# undef SQLITE_WASI
+# define SQLITE_WASI 1
+# ifndef SQLITE_OMIT_LOAD_EXTENSION
+#  define SQLITE_OMIT_LOAD_EXTENSION
+# endif
+# ifndef SQLITE_THREADSAFE
+#  define SQLITE_THREADSAFE 0
+# endif
+#endif
+
 #ifdef __cplusplus
 }  /* End of the 'extern "C"' block */
 #endif
@@ -9785,7 +10952,7 @@ struct sqlite3_rtree_geometry {
 };
 
 /*
-** Register a 2nd-generation geometry callback named zScore that can be 
+** Register a 2nd-generation geometry callback named zScore that can be
 ** used as part of an R-Tree geometry query as follows:
 **
 **   SELECT ... FROM <rtree> WHERE <rtree col> MATCH $zQueryFunc(... params ...)
@@ -9800,7 +10967,7 @@ SQLITE_API int sqlite3_rtree_query_callback(
 
 
 /*
-** A pointer to a structure of the following type is passed as the 
+** A pointer to a structure of the following type is passed as the
 ** argument to scored geometry callback registered using
 ** sqlite3_rtree_query_callback().
 **
@@ -9895,7 +11062,7 @@ typedef struct sqlite3_changeset_iter sqlite3_changeset_iter;
 ** is not possible for an application to register a pre-update hook on a
 ** database handle that has one or more session objects attached. Nor is
 ** it possible to create a session object attached to a database handle for
-** which a pre-update hook is already defined. The results of attempting 
+** which a pre-update hook is already defined. The results of attempting
 ** either of these things are undefined.
 **
 ** The session object will be used to create changesets for tables in
@@ -9913,17 +11080,62 @@ SQLITE_API int sqlite3session_create(
 ** CAPI3REF: Delete A Session Object
 ** DESTRUCTOR: sqlite3_session
 **
-** Delete a session object previously allocated using 
+** Delete a session object previously allocated using
 ** [sqlite3session_create()]. Once a session object has been deleted, the
 ** results of attempting to use pSession with any other session module
 ** function are undefined.
 **
 ** Session objects must be deleted before the database handle to which they
-** are attached is closed. Refer to the documentation for 
+** are attached is closed. Refer to the documentation for
 ** [sqlite3session_create()] for details.
 */
 SQLITE_API void sqlite3session_delete(sqlite3_session *pSession);
 
+/*
+** CAPI3REF: Configure a Session Object
+** METHOD: sqlite3_session
+**
+** This method is used to configure a session object after it has been
+** created. At present the only valid values for the second parameter are
+** [SQLITE_SESSION_OBJCONFIG_SIZE] and [SQLITE_SESSION_OBJCONFIG_ROWID].
+**
+*/
+SQLITE_API int sqlite3session_object_config(sqlite3_session*, int op, void *pArg);
+
+/*
+** CAPI3REF: Options for sqlite3session_object_config
+**
+** The following values may passed as the the 2nd parameter to
+** sqlite3session_object_config().
+**
+** <dt>SQLITE_SESSION_OBJCONFIG_SIZE <dd>
+**   This option is used to set, clear or query the flag that enables
+**   the [sqlite3session_changeset_size()] API. Because it imposes some
+**   computational overhead, this API is disabled by default. Argument
+**   pArg must point to a value of type (int). If the value is initially
+**   0, then the sqlite3session_changeset_size() API is disabled. If it
+**   is greater than 0, then the same API is enabled. Or, if the initial
+**   value is less than zero, no change is made. In all cases the (int)
+**   variable is set to 1 if the sqlite3session_changeset_size() API is
+**   enabled following the current call, or 0 otherwise.
+**
+**   It is an error (SQLITE_MISUSE) to attempt to modify this setting after
+**   the first table has been attached to the session object.
+**
+** <dt>SQLITE_SESSION_OBJCONFIG_ROWID <dd>
+**   This option is used to set, clear or query the flag that enables
+**   collection of data for tables with no explicit PRIMARY KEY.
+**
+**   Normally, tables with no explicit PRIMARY KEY are simply ignored
+**   by the sessions module. However, if this flag is set, it behaves
+**   as if such tables have a column "_rowid_ INTEGER PRIMARY KEY" inserted
+**   as their leftmost columns.
+**
+**   It is an error (SQLITE_MISUSE) to attempt to modify this setting after
+**   the first table has been attached to the session object.
+*/
+#define SQLITE_SESSION_OBJCONFIG_SIZE  1
+#define SQLITE_SESSION_OBJCONFIG_ROWID 2
 
 /*
 ** CAPI3REF: Enable Or Disable A Session Object
@@ -9937,10 +11149,10 @@ SQLITE_API void sqlite3session_delete(sqlite3_session *pSession);
 ** the eventual changesets.
 **
 ** Passing zero to this function disables the session. Passing a value
-** greater than zero enables it. Passing a value less than zero is a 
+** greater than zero enables it. Passing a value less than zero is a
 ** no-op, and may be used to query the current state of the session.
 **
-** The return value indicates the final state of the session object: 0 if 
+** The return value indicates the final state of the session object: 0 if
 ** the session is disabled, or 1 if it is enabled.
 */
 SQLITE_API int sqlite3session_enable(sqlite3_session *pSession, int bEnable);
@@ -9955,7 +11167,7 @@ SQLITE_API int sqlite3session_enable(sqlite3_session *pSession, int bEnable);
 ** <ul>
 **   <li> The session object "indirect" flag is set when the change is
 **        made, or
-**   <li> The change is made by an SQL trigger or foreign key action 
+**   <li> The change is made by an SQL trigger or foreign key action
 **        instead of directly as a result of a users SQL statement.
 ** </ul>
 **
@@ -9967,10 +11179,10 @@ SQLITE_API int sqlite3session_enable(sqlite3_session *pSession, int bEnable);
 ** flag.  If the second argument passed to this function is zero, then the
 ** indirect flag is cleared. If it is greater than zero, the indirect flag
 ** is set. Passing a value less than zero does not modify the current value
-** of the indirect flag, and may be used to query the current state of the 
+** of the indirect flag, and may be used to query the current state of the
 ** indirect flag for the specified session object.
 **
-** The return value indicates the final state of the indirect flag: 0 if 
+** The return value indicates the final state of the indirect flag: 0 if
 ** it is clear, or 1 if it is set.
 */
 SQLITE_API int sqlite3session_indirect(sqlite3_session *pSession, int bIndirect);
@@ -9980,20 +11192,20 @@ SQLITE_API int sqlite3session_indirect(sqlite3_session *pSession, int bIndirect)
 ** METHOD: sqlite3_session
 **
 ** If argument zTab is not NULL, then it is the name of a table to attach
-** to the session object passed as the first argument. All subsequent changes 
-** made to the table while the session object is enabled will be recorded. See 
+** to the session object passed as the first argument. All subsequent changes
+** made to the table while the session object is enabled will be recorded. See
 ** documentation for [sqlite3session_changeset()] for further details.
 **
 ** Or, if argument zTab is NULL, then changes are recorded for all tables
-** in the database. If additional tables are added to the database (by 
-** executing "CREATE TABLE" statements) after this call is made, changes for 
+** in the database. If additional tables are added to the database (by
+** executing "CREATE TABLE" statements) after this call is made, changes for
 ** the new tables are also recorded.
 **
 ** Changes can only be recorded for tables that have a PRIMARY KEY explicitly
-** defined as part of their CREATE TABLE statement. It does not matter if the 
+** defined as part of their CREATE TABLE statement. It does not matter if the
 ** PRIMARY KEY is an "INTEGER PRIMARY KEY" (rowid alias) or not. The PRIMARY
 ** KEY may consist of a single column, or may be a composite key.
-** 
+**
 ** It is not an error if the named table does not exist in the database. Nor
 ** is it an error if the named table does not have a PRIMARY KEY. However,
 ** no changes will be recorded in either of these scenarios.
@@ -10001,29 +11213,29 @@ SQLITE_API int sqlite3session_indirect(sqlite3_session *pSession, int bIndirect)
 ** Changes are not recorded for individual rows that have NULL values stored
 ** in one or more of their PRIMARY KEY columns.
 **
-** SQLITE_OK is returned if the call completes without error. Or, if an error 
+** SQLITE_OK is returned if the call completes without error. Or, if an error
 ** occurs, an SQLite error code (e.g. SQLITE_NOMEM) is returned.
 **
 ** <h3>Special sqlite_stat1 Handling</h3>
 **
-** As of SQLite version 3.22.0, the "sqlite_stat1" table is an exception to 
+** As of SQLite version 3.22.0, the "sqlite_stat1" table is an exception to
 ** some of the rules above. In SQLite, the schema of sqlite_stat1 is:
 **  <pre>
-**  &nbsp;     CREATE TABLE sqlite_stat1(tbl,idx,stat)  
+**  &nbsp;     CREATE TABLE sqlite_stat1(tbl,idx,stat)
 **  </pre>
 **
-** Even though sqlite_stat1 does not have a PRIMARY KEY, changes are 
-** recorded for it as if the PRIMARY KEY is (tbl,idx). Additionally, changes 
+** Even though sqlite_stat1 does not have a PRIMARY KEY, changes are
+** recorded for it as if the PRIMARY KEY is (tbl,idx). Additionally, changes
 ** are recorded for rows for which (idx IS NULL) is true. However, for such
 ** rows a zero-length blob (SQL value X'') is stored in the changeset or
 ** patchset instead of a NULL value. This allows such changesets to be
 ** manipulated by legacy implementations of sqlite3changeset_invert(),
 ** concat() and similar.
 **
-** The sqlite3changeset_apply() function automatically converts the 
+** The sqlite3changeset_apply() function automatically converts the
 ** zero-length blob back to a NULL value when updating the sqlite_stat1
 ** table. However, if the application calls sqlite3changeset_new(),
-** sqlite3changeset_old() or sqlite3changeset_conflict on a changeset 
+** sqlite3changeset_old() or sqlite3changeset_conflict on a changeset
 ** iterator directly (including on a changeset iterator passed to a
 ** conflict-handler callback) then the X'' value is returned. The application
 ** must translate X'' to NULL itself if required.
@@ -10042,10 +11254,10 @@ SQLITE_API int sqlite3session_attach(
 ** CAPI3REF: Set a table filter on a Session Object.
 ** METHOD: sqlite3_session
 **
-** The second argument (xFilter) is the "filter callback". For changes to rows 
+** The second argument (xFilter) is the "filter callback". For changes to rows
 ** in tables that are not attached to the Session object, the filter is called
-** to determine whether changes to the table's rows should be tracked or not. 
-** If xFilter returns 0, changes are not tracked. Note that once a table is 
+** to determine whether changes to the table's rows should be tracked or not.
+** If xFilter returns 0, changes are not tracked. Note that once a table is
 ** attached, xFilter will not be called again.
 */
 SQLITE_API void sqlite3session_table_filter(
@@ -10061,9 +11273,9 @@ SQLITE_API void sqlite3session_table_filter(
 ** CAPI3REF: Generate A Changeset From A Session Object
 ** METHOD: sqlite3_session
 **
-** Obtain a changeset containing changes to the tables attached to the 
-** session object passed as the first argument. If successful, 
-** set *ppChangeset to point to a buffer containing the changeset 
+** Obtain a changeset containing changes to the tables attached to the
+** session object passed as the first argument. If successful,
+** set *ppChangeset to point to a buffer containing the changeset
 ** and *pnChangeset to the size of the changeset in bytes before returning
 ** SQLITE_OK. If an error occurs, set both *ppChangeset and *pnChangeset to
 ** zero and return an SQLite error code.
@@ -10078,7 +11290,7 @@ SQLITE_API void sqlite3session_table_filter(
 ** modifies the values of primary key columns. If such a change is made, it
 ** is represented in a changeset as a DELETE followed by an INSERT.
 **
-** Changes are not recorded for rows that have NULL values stored in one or 
+** Changes are not recorded for rows that have NULL values stored in one or
 ** more of their PRIMARY KEY columns. If such a row is inserted or deleted,
 ** no corresponding change is present in the changesets returned by this
 ** function. If an existing row with one or more NULL values stored in
@@ -10131,14 +11343,14 @@ SQLITE_API void sqlite3session_table_filter(
 ** <ul>
 **   <li> For each record generated by an insert, the database is queried
 **        for a row with a matching primary key. If one is found, an INSERT
-**        change is added to the changeset. If no such row is found, no change 
+**        change is added to the changeset. If no such row is found, no change
 **        is added to the changeset.
 **
-**   <li> For each record generated by an update or delete, the database is 
+**   <li> For each record generated by an update or delete, the database is
 **        queried for a row with a matching primary key. If such a row is
 **        found and one or more of the non-primary key fields have been
-**        modified from their original values, an UPDATE change is added to 
-**        the changeset. Or, if no such row is found in the table, a DELETE 
+**        modified from their original values, an UPDATE change is added to
+**        the changeset. Or, if no such row is found in the table, a DELETE
 **        change is added to the changeset. If there is a row with a matching
 **        primary key in the database, but all fields contain their original
 **        values, no change is added to the changeset.
@@ -10146,7 +11358,7 @@ SQLITE_API void sqlite3session_table_filter(
 **
 ** This means, amongst other things, that if a row is inserted and then later
 ** deleted while a session object is active, neither the insert nor the delete
-** will be present in the changeset. Or if a row is deleted and then later a 
+** will be present in the changeset. Or if a row is deleted and then later a
 ** row with the same primary key values inserted while a session object is
 ** active, the resulting changeset will contain an UPDATE change instead of
 ** a DELETE and an INSERT.
@@ -10155,10 +11367,10 @@ SQLITE_API void sqlite3session_table_filter(
 ** it does not accumulate records when rows are inserted, updated or deleted.
 ** This may appear to have some counter-intuitive effects if a single row
 ** is written to more than once during a session. For example, if a row
-** is inserted while a session object is enabled, then later deleted while 
+** is inserted while a session object is enabled, then later deleted while
 ** the same session object is disabled, no INSERT record will appear in the
 ** changeset, even though the delete took place while the session was disabled.
-** Or, if one field of a row is updated while a session is disabled, and 
+** Or, if one field of a row is updated while a session is disabled, and
 ** another field of the same row is updated while the session is enabled, the
 ** resulting changeset will contain an UPDATE change that updates both fields.
 */
@@ -10168,6 +11380,22 @@ SQLITE_API int sqlite3session_changeset(
   void **ppChangeset              /* OUT: Buffer containing changeset */
 );
 
+/*
+** CAPI3REF: Return An Upper-limit For The Size Of The Changeset
+** METHOD: sqlite3_session
+**
+** By default, this function always returns 0. For it to return
+** a useful result, the sqlite3_session object must have been configured
+** to enable this API using sqlite3session_object_config() with the
+** SQLITE_SESSION_OBJCONFIG_SIZE verb.
+**
+** When enabled, this function returns an upper limit, in bytes, for the size
+** of the changeset that might be produced if sqlite3session_changeset() were
+** called. The final changeset size might be equal to or smaller than the
+** size in bytes returned by this function.
+*/
+SQLITE_API sqlite3_int64 sqlite3session_changeset_size(sqlite3_session *pSession);
+
 /*
 ** CAPI3REF: Load The Difference Between Tables Into A Session
 ** METHOD: sqlite3_session
@@ -10179,7 +11407,7 @@ SQLITE_API int sqlite3session_changeset(
 ** an error).
 **
 ** Argument zFromDb must be the name of a database ("main", "temp" etc.)
-** attached to the same database handle as the session object that contains 
+** attached to the same database handle as the session object that contains
 ** a table compatible with the table attached to the session by this function.
 ** A table is considered compatible if it:
 **
@@ -10195,25 +11423,25 @@ SQLITE_API int sqlite3session_changeset(
 ** APIs, tables without PRIMARY KEYs are simply ignored.
 **
 ** This function adds a set of changes to the session object that could be
-** used to update the table in database zFrom (call this the "from-table") 
-** so that its content is the same as the table attached to the session 
+** used to update the table in database zFrom (call this the "from-table")
+** so that its content is the same as the table attached to the session
 ** object (call this the "to-table"). Specifically:
 **
 ** <ul>
-**   <li> For each row (primary key) that exists in the to-table but not in 
+**   <li> For each row (primary key) that exists in the to-table but not in
 **     the from-table, an INSERT record is added to the session object.
 **
-**   <li> For each row (primary key) that exists in the to-table but not in 
+**   <li> For each row (primary key) that exists in the to-table but not in
 **     the from-table, a DELETE record is added to the session object.
 **
-**   <li> For each row (primary key) that exists in both tables, but features 
+**   <li> For each row (primary key) that exists in both tables, but features
 **     different non-PK values in each, an UPDATE record is added to the
-**     session.  
+**     session.
 ** </ul>
 **
 ** To clarify, if this function is called and then a changeset constructed
-** using [sqlite3session_changeset()], then after applying that changeset to 
-** database zFrom the contents of the two compatible tables would be 
+** using [sqlite3session_changeset()], then after applying that changeset to
+** database zFrom the contents of the two compatible tables would be
 ** identical.
 **
 ** It an error if database zFrom does not exist or does not contain the
@@ -10221,7 +11449,7 @@ SQLITE_API int sqlite3session_changeset(
 **
 ** If the operation is successful, SQLITE_OK is returned. Otherwise, an SQLite
 ** error code. In this case, if argument pzErrMsg is not NULL, *pzErrMsg
-** may be set to point to a buffer containing an English language error 
+** may be set to point to a buffer containing an English language error
 ** message. It is the responsibility of the caller to free this buffer using
 ** sqlite3_free().
 */
@@ -10240,19 +11468,19 @@ SQLITE_API int sqlite3session_diff(
 ** The differences between a patchset and a changeset are that:
 **
 ** <ul>
-**   <li> DELETE records consist of the primary key fields only. The 
+**   <li> DELETE records consist of the primary key fields only. The
 **        original values of other fields are omitted.
-**   <li> The original values of any modified fields are omitted from 
+**   <li> The original values of any modified fields are omitted from
 **        UPDATE records.
 ** </ul>
 **
-** A patchset blob may be used with up to date versions of all 
-** sqlite3changeset_xxx API functions except for sqlite3changeset_invert(), 
+** A patchset blob may be used with up to date versions of all
+** sqlite3changeset_xxx API functions except for sqlite3changeset_invert(),
 ** which returns SQLITE_CORRUPT if it is passed a patchset. Similarly,
 ** attempting to use a patchset blob with old versions of the
-** sqlite3changeset_xxx APIs also provokes an SQLITE_CORRUPT error. 
+** sqlite3changeset_xxx APIs also provokes an SQLITE_CORRUPT error.
 **
-** Because the non-primary key "old.*" fields are omitted, no 
+** Because the non-primary key "old.*" fields are omitted, no
 ** SQLITE_CHANGESET_DATA conflicts can be detected or reported if a patchset
 ** is passed to the sqlite3changeset_apply() API. Other conflict types work
 ** in the same way as for changesets.
@@ -10271,22 +11499,30 @@ SQLITE_API int sqlite3session_patchset(
 /*
 ** CAPI3REF: Test if a changeset has recorded any changes.
 **
-** Return non-zero if no changes to attached tables have been recorded by 
-** the session object passed as the first argument. Otherwise, if one or 
+** Return non-zero if no changes to attached tables have been recorded by
+** the session object passed as the first argument. Otherwise, if one or
 ** more changes have been recorded, return zero.
 **
 ** Even if this function returns zero, it is possible that calling
 ** [sqlite3session_changeset()] on the session handle may still return a
-** changeset that contains no changes. This can happen when a row in 
-** an attached table is modified and then later on the original values 
+** changeset that contains no changes. This can happen when a row in
+** an attached table is modified and then later on the original values
 ** are restored. However, if this function returns non-zero, then it is
-** guaranteed that a call to sqlite3session_changeset() will return a 
+** guaranteed that a call to sqlite3session_changeset() will return a
 ** changeset containing zero changes.
 */
 SQLITE_API int sqlite3session_isempty(sqlite3_session *pSession);
 
 /*
-** CAPI3REF: Create An Iterator To Traverse A Changeset 
+** CAPI3REF: Query for the amount of heap memory used by a session object.
+**
+** This API returns the total amount of heap memory in bytes currently
+** used by the session object passed as the only argument.
+*/
+SQLITE_API sqlite3_int64 sqlite3session_memory_used(sqlite3_session *pSession);
+
+/*
+** CAPI3REF: Create An Iterator To Traverse A Changeset
 ** CONSTRUCTOR: sqlite3_changeset_iter
 **
 ** Create an iterator used to iterate through the contents of a changeset.
@@ -10294,7 +11530,7 @@ SQLITE_API int sqlite3session_isempty(sqlite3_session *pSession);
 ** is returned. Otherwise, if an error occurs, *pp is set to zero and an
 ** SQLite error code is returned.
 **
-** The following functions can be used to advance and query a changeset 
+** The following functions can be used to advance and query a changeset
 ** iterator created by this function:
 **
 ** <ul>
@@ -10311,12 +11547,12 @@ SQLITE_API int sqlite3session_isempty(sqlite3_session *pSession);
 **
 ** Assuming the changeset blob was created by one of the
 ** [sqlite3session_changeset()], [sqlite3changeset_concat()] or
-** [sqlite3changeset_invert()] functions, all changes within the changeset 
-** that apply to a single table are grouped together. This means that when 
-** an application iterates through a changeset using an iterator created by 
-** this function, all changes that relate to a single table are visited 
-** consecutively. There is no chance that the iterator will visit a change 
-** the applies to table X, then one for table Y, and then later on visit 
+** [sqlite3changeset_invert()] functions, all changes within the changeset
+** that apply to a single table are grouped together. This means that when
+** an application iterates through a changeset using an iterator created by
+** this function, all changes that relate to a single table are visited
+** consecutively. There is no chance that the iterator will visit a change
+** the applies to table X, then one for table Y, and then later on visit
 ** another change for table X.
 **
 ** The behavior of sqlite3changeset_start_v2() and its streaming equivalent
@@ -10367,12 +11603,12 @@ SQLITE_API int sqlite3changeset_start_v2(
 ** point to the first change in the changeset. Each subsequent call advances
 ** the iterator to point to the next change in the changeset (if any). If
 ** no error occurs and the iterator points to a valid change after a call
-** to sqlite3changeset_next() has advanced it, SQLITE_ROW is returned. 
+** to sqlite3changeset_next() has advanced it, SQLITE_ROW is returned.
 ** Otherwise, if all changes in the changeset have already been visited,
 ** SQLITE_DONE is returned.
 **
-** If an error occurs, an SQLite error code is returned. Possible error 
-** codes include SQLITE_CORRUPT (if the changeset buffer is corrupt) or 
+** If an error occurs, an SQLite error code is returned. Possible error
+** codes include SQLITE_CORRUPT (if the changeset buffer is corrupt) or
 ** SQLITE_NOMEM.
 */
 SQLITE_API int sqlite3changeset_next(sqlite3_changeset_iter *pIter);
@@ -10387,18 +11623,23 @@ SQLITE_API int sqlite3changeset_next(sqlite3_changeset_iter *pIter);
 ** call to [sqlite3changeset_next()] must have returned [SQLITE_ROW]. If this
 ** is not the case, this function returns [SQLITE_MISUSE].
 **
-** If argument pzTab is not NULL, then *pzTab is set to point to a
-** nul-terminated utf-8 encoded string containing the name of the table
-** affected by the current change. The buffer remains valid until either
-** sqlite3changeset_next() is called on the iterator or until the 
-** conflict-handler function returns. If pnCol is not NULL, then *pnCol is 
-** set to the number of columns in the table affected by the change. If
-** pbIndirect is not NULL, then *pbIndirect is set to true (1) if the change
+** Arguments pOp, pnCol and pzTab may not be NULL. Upon return, three
+** outputs are set through these pointers:
+**
+** *pOp is set to one of [SQLITE_INSERT], [SQLITE_DELETE] or [SQLITE_UPDATE],
+** depending on the type of change that the iterator currently points to;
+**
+** *pnCol is set to the number of columns in the table affected by the change; and
+**
+** *pzTab is set to point to a nul-terminated utf-8 encoded string containing
+** the name of the table affected by the current change. The buffer remains
+** valid until either sqlite3changeset_next() is called on the iterator
+** or until the conflict-handler function returns.
+**
+** If pbIndirect is not NULL, then *pbIndirect is set to true (1) if the change
 ** is an indirect change, or false (0) otherwise. See the documentation for
 ** [sqlite3session_indirect()] for a description of direct and indirect
-** changes. Finally, if pOp is not NULL, then *pOp is set to one of 
-** [SQLITE_INSERT], [SQLITE_DELETE] or [SQLITE_UPDATE], depending on the 
-** type of change that the iterator currently points to.
+** changes.
 **
 ** If no error occurs, SQLITE_OK is returned. If an error does occur, an
 ** SQLite error code is returned. The values of the output variables may not
@@ -10451,7 +11692,7 @@ SQLITE_API int sqlite3changeset_pk(
 ** The pIter argument passed to this function may either be an iterator
 ** passed to a conflict-handler by [sqlite3changeset_apply()], or an iterator
 ** created by [sqlite3changeset_start()]. In the latter case, the most recent
-** call to [sqlite3changeset_next()] must have returned SQLITE_ROW. 
+** call to [sqlite3changeset_next()] must have returned SQLITE_ROW.
 ** Furthermore, it may only be called if the type of change that the iterator
 ** currently points to is either [SQLITE_DELETE] or [SQLITE_UPDATE]. Otherwise,
 ** this function returns [SQLITE_MISUSE] and sets *ppValue to NULL.
@@ -10461,9 +11702,9 @@ SQLITE_API int sqlite3changeset_pk(
 ** [SQLITE_RANGE] is returned and *ppValue is set to NULL.
 **
 ** If successful, this function sets *ppValue to point to a protected
-** sqlite3_value object containing the iVal'th value from the vector of 
+** sqlite3_value object containing the iVal'th value from the vector of
 ** original row values stored as part of the UPDATE or DELETE change and
-** returns SQLITE_OK. The name of the function comes from the fact that this 
+** returns SQLITE_OK. The name of the function comes from the fact that this
 ** is similar to the "old.*" columns available to update or delete triggers.
 **
 ** If some other error occurs (e.g. an OOM condition), an SQLite error code
@@ -10482,7 +11723,7 @@ SQLITE_API int sqlite3changeset_old(
 ** The pIter argument passed to this function may either be an iterator
 ** passed to a conflict-handler by [sqlite3changeset_apply()], or an iterator
 ** created by [sqlite3changeset_start()]. In the latter case, the most recent
-** call to [sqlite3changeset_next()] must have returned SQLITE_ROW. 
+** call to [sqlite3changeset_next()] must have returned SQLITE_ROW.
 ** Furthermore, it may only be called if the type of change that the iterator
 ** currently points to is either [SQLITE_UPDATE] or [SQLITE_INSERT]. Otherwise,
 ** this function returns [SQLITE_MISUSE] and sets *ppValue to NULL.
@@ -10492,12 +11733,12 @@ SQLITE_API int sqlite3changeset_old(
 ** [SQLITE_RANGE] is returned and *ppValue is set to NULL.
 **
 ** If successful, this function sets *ppValue to point to a protected
-** sqlite3_value object containing the iVal'th value from the vector of 
+** sqlite3_value object containing the iVal'th value from the vector of
 ** new row values stored as part of the UPDATE or INSERT change and
 ** returns SQLITE_OK. If the change is an UPDATE and does not include
-** a new value for the requested column, *ppValue is set to NULL and 
-** SQLITE_OK returned. The name of the function comes from the fact that 
-** this is similar to the "new.*" columns available to update or delete 
+** a new value for the requested column, *ppValue is set to NULL and
+** SQLITE_OK returned. The name of the function comes from the fact that
+** this is similar to the "new.*" columns available to update or delete
 ** triggers.
 **
 ** If some other error occurs (e.g. an OOM condition), an SQLite error code
@@ -10524,7 +11765,7 @@ SQLITE_API int sqlite3changeset_new(
 ** [SQLITE_RANGE] is returned and *ppValue is set to NULL.
 **
 ** If successful, this function sets *ppValue to point to a protected
-** sqlite3_value object containing the iVal'th value from the 
+** sqlite3_value object containing the iVal'th value from the
 ** "conflicting row" associated with the current conflict-handler callback
 ** and returns SQLITE_OK.
 **
@@ -10568,7 +11809,7 @@ SQLITE_API int sqlite3changeset_fk_conflicts(
 ** call has no effect.
 **
 ** If an error was encountered within a call to an sqlite3changeset_xxx()
-** function (for example an [SQLITE_CORRUPT] in [sqlite3changeset_next()] or an 
+** function (for example an [SQLITE_CORRUPT] in [sqlite3changeset_next()] or an
 ** [SQLITE_NOMEM] in [sqlite3changeset_new()]) then an error code corresponding
 ** to that error is returned by this function. Otherwise, SQLITE_OK is
 ** returned. This is to allow the following pattern (pseudo-code):
@@ -10580,7 +11821,7 @@ SQLITE_API int sqlite3changeset_fk_conflicts(
 **   }
 **   rc = sqlite3changeset_finalize();
 **   if( rc!=SQLITE_OK ){
-**     // An error has occurred 
+**     // An error has occurred
 **   }
 ** </pre>
 */
@@ -10608,7 +11849,7 @@ SQLITE_API int sqlite3changeset_finalize(sqlite3_changeset_iter *pIter);
 ** zeroed and an SQLite error code returned.
 **
 ** It is the responsibility of the caller to eventually call sqlite3_free()
-** on the *ppOut pointer to free the buffer allocation following a successful 
+** on the *ppOut pointer to free the buffer allocation following a successful
 ** call to this function.
 **
 ** WARNING/TODO: This function currently assumes that the input is a valid
@@ -10622,11 +11863,11 @@ SQLITE_API int sqlite3changeset_invert(
 /*
 ** CAPI3REF: Concatenate Two Changeset Objects
 **
-** This function is used to concatenate two changesets, A and B, into a 
+** This function is used to concatenate two changesets, A and B, into a
 ** single changeset. The result is a changeset equivalent to applying
-** changeset A followed by changeset B. 
+** changeset A followed by changeset B.
 **
-** This function combines the two input changesets using an 
+** This function combines the two input changesets using an
 ** sqlite3_changegroup object. Calling it produces similar results as the
 ** following code fragment:
 **
@@ -10655,10 +11896,22 @@ SQLITE_API int sqlite3changeset_concat(
 );
 
 
+/*
+** CAPI3REF: Upgrade the Schema of a Changeset/Patchset
+*/
+SQLITE_API int sqlite3changeset_upgrade(
+  sqlite3 *db,
+  const char *zDb,
+  int nIn, const void *pIn,       /* Input changeset */
+  int *pnOut, void **ppOut        /* OUT: Inverse of input */
+);
+
+
+
 /*
 ** CAPI3REF: Changegroup Handle
 **
-** A changegroup is an object used to combine two or more 
+** A changegroup is an object used to combine two or more
 ** [changesets] or [patchsets]
 */
 typedef struct sqlite3_changegroup sqlite3_changegroup;
@@ -10674,7 +11927,7 @@ typedef struct sqlite3_changegroup sqlite3_changegroup;
 **
 ** If successful, this function returns SQLITE_OK and populates (*pp) with
 ** a pointer to a new sqlite3_changegroup object before returning. The caller
-** should eventually free the returned object using a call to 
+** should eventually free the returned object using a call to
 ** sqlite3changegroup_delete(). If an error occurs, an SQLite error code
 ** (i.e. SQLITE_NOMEM) is returned and *pp is set to NULL.
 **
@@ -10686,7 +11939,7 @@ typedef struct sqlite3_changegroup sqlite3_changegroup;
 **   <li> Zero or more changesets (or patchsets) are added to the object
 **        by calling sqlite3changegroup_add().
 **
-**   <li> The result of combining all input changesets together is obtained 
+**   <li> The result of combining all input changesets together is obtained
 **        by the application via a call to sqlite3changegroup_output().
 **
 **   <li> The object is deleted using a call to sqlite3changegroup_delete().
@@ -10695,18 +11948,50 @@ typedef struct sqlite3_changegroup sqlite3_changegroup;
 ** Any number of calls to add() and output() may be made between the calls to
 ** new() and delete(), and in any order.
 **
-** As well as the regular sqlite3changegroup_add() and 
+** As well as the regular sqlite3changegroup_add() and
 ** sqlite3changegroup_output() functions, also available are the streaming
 ** versions sqlite3changegroup_add_strm() and sqlite3changegroup_output_strm().
 */
 SQLITE_API int sqlite3changegroup_new(sqlite3_changegroup **pp);
 
+/*
+** CAPI3REF: Add a Schema to a Changegroup
+** METHOD: sqlite3_changegroup_schema
+**
+** This method may be used to optionally enforce the rule that the changesets
+** added to the changegroup handle must match the schema of database zDb
+** ("main", "temp", or the name of an attached database). If
+** sqlite3changegroup_add() is called to add a changeset that is not compatible
+** with the configured schema, SQLITE_SCHEMA is returned and the changegroup
+** object is left in an undefined state.
+**
+** A changeset schema is considered compatible with the database schema in
+** the same way as for sqlite3changeset_apply(). Specifically, for each
+** table in the changeset, there exists a database table with:
+**
+** <ul>
+**   <li> The name identified by the changeset, and
+**   <li> at least as many columns as recorded in the changeset, and
+**   <li> the primary key columns in the same position as recorded in
+**        the changeset.
+** </ul>
+**
+** The output of the changegroup object always has the same schema as the
+** database nominated using this function. In cases where changesets passed
+** to sqlite3changegroup_add() have fewer columns than the corresponding table
+** in the database schema, these are filled in using the default column
+** values from the database schema. This makes it possible to combined
+** changesets that have different numbers of columns for a single table
+** within a changegroup, provided that they are otherwise compatible.
+*/
+SQLITE_API int sqlite3changegroup_schema(sqlite3_changegroup*, sqlite3*, const char *zDb);
+
 /*
 ** CAPI3REF: Add A Changeset To A Changegroup
 ** METHOD: sqlite3_changegroup
 **
 ** Add all changes within the changeset (or patchset) in buffer pData (size
-** nData bytes) to the changegroup. 
+** nData bytes) to the changegroup.
 **
 ** If the buffer contains a patchset, then all prior calls to this function
 ** on the same changegroup object must also have specified patchsets. Or, if
@@ -10733,7 +12018,7 @@ SQLITE_API int sqlite3changegroup_new(sqlite3_changegroup **pp);
 **       changeset was recorded immediately after the changesets already
 **       added to the changegroup.
 **   <tr><td>INSERT <td>UPDATE <td>
-**       The INSERT change remains in the changegroup. The values in the 
+**       The INSERT change remains in the changegroup. The values in the
 **       INSERT change are modified as if the row was inserted by the
 **       existing change and then updated according to the new change.
 **   <tr><td>INSERT <td>DELETE <td>
@@ -10744,17 +12029,17 @@ SQLITE_API int sqlite3changegroup_new(sqlite3_changegroup **pp);
 **       changeset was recorded immediately after the changesets already
 **       added to the changegroup.
 **   <tr><td>UPDATE <td>UPDATE <td>
-**       The existing UPDATE remains within the changegroup. It is amended 
-**       so that the accompanying values are as if the row was updated once 
+**       The existing UPDATE remains within the changegroup. It is amended
+**       so that the accompanying values are as if the row was updated once
 **       by the existing change and then again by the new change.
 **   <tr><td>UPDATE <td>DELETE <td>
 **       The existing UPDATE is replaced by the new DELETE within the
 **       changegroup.
 **   <tr><td>DELETE <td>INSERT <td>
 **       If one or more of the column values in the row inserted by the
-**       new change differ from those in the row deleted by the existing 
+**       new change differ from those in the row deleted by the existing
 **       change, the existing DELETE is replaced by an UPDATE within the
-**       changegroup. Otherwise, if the inserted row is exactly the same 
+**       changegroup. Otherwise, if the inserted row is exactly the same
 **       as the deleted row, the existing DELETE is simply discarded.
 **   <tr><td>DELETE <td>UPDATE <td>
 **       The new change is ignored. This case does not occur if the new
@@ -10769,16 +12054,45 @@ SQLITE_API int sqlite3changegroup_new(sqlite3_changegroup **pp);
 ** If the new changeset contains changes to a table that is already present
 ** in the changegroup, then the number of columns and the position of the
 ** primary key columns for the table must be consistent. If this is not the
-** case, this function fails with SQLITE_SCHEMA. If the input changeset
-** appears to be corrupt and the corruption is detected, SQLITE_CORRUPT is
-** returned. Or, if an out-of-memory condition occurs during processing, this
-** function returns SQLITE_NOMEM. In all cases, if an error occurs the state
-** of the final contents of the changegroup is undefined.
+** case, this function fails with SQLITE_SCHEMA. Except, if the changegroup
+** object has been configured with a database schema using the
+** sqlite3changegroup_schema() API, then it is possible to combine changesets
+** with different numbers of columns for a single table, provided that
+** they are otherwise compatible.
+**
+** If the input changeset appears to be corrupt and the corruption is
+** detected, SQLITE_CORRUPT is returned. Or, if an out-of-memory condition
+** occurs during processing, this function returns SQLITE_NOMEM.
 **
-** If no error occurs, SQLITE_OK is returned.
+** In all cases, if an error occurs the state of the final contents of the
+** changegroup is undefined. If no error occurs, SQLITE_OK is returned.
 */
 SQLITE_API int sqlite3changegroup_add(sqlite3_changegroup*, int nData, void *pData);
 
+/*
+** CAPI3REF: Add A Single Change To A Changegroup
+** METHOD: sqlite3_changegroup
+**
+** This function adds the single change currently indicated by the iterator
+** passed as the second argument to the changegroup object. The rules for
+** adding the change are just as described for [sqlite3changegroup_add()].
+**
+** If the change is successfully added to the changegroup, SQLITE_OK is
+** returned. Otherwise, an SQLite error code is returned.
+**
+** The iterator must point to a valid entry when this function is called.
+** If it does not, SQLITE_ERROR is returned and no change is added to the
+** changegroup. Additionally, the iterator must not have been opened with
+** the SQLITE_CHANGESETAPPLY_INVERT flag. In this case SQLITE_ERROR is also
+** returned.
+*/
+SQLITE_API int sqlite3changegroup_add_change(
+  sqlite3_changegroup*,
+  sqlite3_changeset_iter*
+);
+
+
+
 /*
 ** CAPI3REF: Obtain A Composite Changeset From A Changegroup
 ** METHOD: sqlite3_changegroup
@@ -10799,7 +12113,7 @@ SQLITE_API int sqlite3changegroup_add(sqlite3_changegroup*, int nData, void *pDa
 **
 ** If an error occurs, an SQLite error code is returned and the output
 ** variables (*pnData) and (*ppData) are set to 0. Otherwise, SQLITE_OK
-** is returned and the output variables are set to the size of and a 
+** is returned and the output variables are set to the size of and a
 ** pointer to the output buffer, respectively. In this case it is the
 ** responsibility of the caller to eventually free the buffer using a
 ** call to sqlite3_free().
@@ -10821,7 +12135,7 @@ SQLITE_API void sqlite3changegroup_delete(sqlite3_changegroup*);
 **
 ** Apply a changeset or patchset to a database. These functions attempt to
 ** update the "main" database attached to handle db with the changes found in
-** the changeset passed via the second and third arguments. 
+** the changeset passed via the second and third arguments.
 **
 ** The fourth argument (xFilter) passed to these functions is the "filter
 ** callback". If it is not NULL, then for each table affected by at least one
@@ -10832,16 +12146,16 @@ SQLITE_API void sqlite3changegroup_delete(sqlite3_changegroup*);
 ** Otherwise, if the return value is non-zero or the xFilter argument to
 ** is NULL, all changes related to the table are attempted.
 **
-** For each table that is not excluded by the filter callback, this function 
-** tests that the target database contains a compatible table. A table is 
+** For each table that is not excluded by the filter callback, this function
+** tests that the target database contains a compatible table. A table is
 ** considered compatible if all of the following are true:
 **
 ** <ul>
-**   <li> The table has the same name as the name recorded in the 
+**   <li> The table has the same name as the name recorded in the
 **        changeset, and
-**   <li> The table has at least as many columns as recorded in the 
+**   <li> The table has at least as many columns as recorded in the
 **        changeset, and
-**   <li> The table has primary key columns in the same position as 
+**   <li> The table has primary key columns in the same position as
 **        recorded in the changeset.
 ** </ul>
 **
@@ -10850,11 +12164,11 @@ SQLITE_API void sqlite3changegroup_delete(sqlite3_changegroup*);
 ** via the sqlite3_log() mechanism with the error code SQLITE_SCHEMA. At most
 ** one such warning is issued for each table in the changeset.
 **
-** For each change for which there is a compatible table, an attempt is made 
-** to modify the table contents according to the UPDATE, INSERT or DELETE 
-** change. If a change cannot be applied cleanly, the conflict handler 
-** function passed as the fifth argument to sqlite3changeset_apply() may be 
-** invoked. A description of exactly when the conflict handler is invoked for 
+** For each change for which there is a compatible table, an attempt is made
+** to modify the table contents according to the UPDATE, INSERT or DELETE
+** change. If a change cannot be applied cleanly, the conflict handler
+** function passed as the fifth argument to sqlite3changeset_apply() may be
+** invoked. A description of exactly when the conflict handler is invoked for
 ** each type of change is below.
 **
 ** Unlike the xFilter argument, xConflict may not be passed NULL. The results
@@ -10862,23 +12176,23 @@ SQLITE_API void sqlite3changegroup_delete(sqlite3_changegroup*);
 ** argument are undefined.
 **
 ** Each time the conflict handler function is invoked, it must return one
-** of [SQLITE_CHANGESET_OMIT], [SQLITE_CHANGESET_ABORT] or 
+** of [SQLITE_CHANGESET_OMIT], [SQLITE_CHANGESET_ABORT] or
 ** [SQLITE_CHANGESET_REPLACE]. SQLITE_CHANGESET_REPLACE may only be returned
 ** if the second argument passed to the conflict handler is either
 ** SQLITE_CHANGESET_DATA or SQLITE_CHANGESET_CONFLICT. If the conflict-handler
 ** returns an illegal value, any changes already made are rolled back and
-** the call to sqlite3changeset_apply() returns SQLITE_MISUSE. Different 
+** the call to sqlite3changeset_apply() returns SQLITE_MISUSE. Different
 ** actions are taken by sqlite3changeset_apply() depending on the value
 ** returned by each invocation of the conflict-handler function. Refer to
-** the documentation for the three 
+** the documentation for the three
 ** [SQLITE_CHANGESET_OMIT|available return values] for details.
 **
 ** <dl>
 ** <dt>DELETE Changes<dd>
-**   For each DELETE change, the function checks if the target database 
-**   contains a row with the same primary key value (or values) as the 
-**   original row values stored in the changeset. If it does, and the values 
-**   stored in all non-primary key columns also match the values stored in 
+**   For each DELETE change, the function checks if the target database
+**   contains a row with the same primary key value (or values) as the
+**   original row values stored in the changeset. If it does, and the values
+**   stored in all non-primary key columns also match the values stored in
 **   the changeset the row is deleted from the target database.
 **
 **   If a row with matching primary key values is found, but one or more of
@@ -10907,22 +12221,22 @@ SQLITE_API void sqlite3changegroup_delete(sqlite3_changegroup*);
 **   database table, the trailing fields are populated with their default
 **   values.
 **
-**   If the attempt to insert the row fails because the database already 
+**   If the attempt to insert the row fails because the database already
 **   contains a row with the same primary key values, the conflict handler
-**   function is invoked with the second argument set to 
+**   function is invoked with the second argument set to
 **   [SQLITE_CHANGESET_CONFLICT].
 **
 **   If the attempt to insert the row fails because of some other constraint
-**   violation (e.g. NOT NULL or UNIQUE), the conflict handler function is 
+**   violation (e.g. NOT NULL or UNIQUE), the conflict handler function is
 **   invoked with the second argument set to [SQLITE_CHANGESET_CONSTRAINT].
-**   This includes the case where the INSERT operation is re-attempted because 
-**   an earlier call to the conflict handler function returned 
+**   This includes the case where the INSERT operation is re-attempted because
+**   an earlier call to the conflict handler function returned
 **   [SQLITE_CHANGESET_REPLACE].
 **
 ** <dt>UPDATE Changes<dd>
-**   For each UPDATE change, the function checks if the target database 
-**   contains a row with the same primary key value (or values) as the 
-**   original row values stored in the changeset. If it does, and the values 
+**   For each UPDATE change, the function checks if the target database
+**   contains a row with the same primary key value (or values) as the
+**   original row values stored in the changeset. If it does, and the values
 **   stored in all modified non-primary key columns also match the values
 **   stored in the changeset the row is updated within the target database.
 **
@@ -10938,12 +12252,12 @@ SQLITE_API void sqlite3changegroup_delete(sqlite3_changegroup*);
 **   the conflict-handler function is invoked with [SQLITE_CHANGESET_NOTFOUND]
 **   passed as the second argument.
 **
-**   If the UPDATE operation is attempted, but SQLite returns 
-**   SQLITE_CONSTRAINT, the conflict-handler function is invoked with 
+**   If the UPDATE operation is attempted, but SQLite returns
+**   SQLITE_CONSTRAINT, the conflict-handler function is invoked with
 **   [SQLITE_CHANGESET_CONSTRAINT] passed as the second argument.
-**   This includes the case where the UPDATE operation is attempted after 
+**   This includes the case where the UPDATE operation is attempted after
 **   an earlier call to the conflict handler function returned
-**   [SQLITE_CHANGESET_REPLACE].  
+**   [SQLITE_CHANGESET_REPLACE].
 ** </dl>
 **
 ** It is safe to execute SQL statements, including those that write to the
@@ -10954,12 +12268,12 @@ SQLITE_API void sqlite3changegroup_delete(sqlite3_changegroup*);
 ** All changes made by these functions are enclosed in a savepoint transaction.
 ** If any other error (aside from a constraint failure when attempting to
 ** write to the target database) occurs, then the savepoint transaction is
-** rolled back, restoring the target database to its original state, and an 
+** rolled back, restoring the target database to its original state, and an
 ** SQLite error code returned.
 **
 ** If the output parameters (ppRebase) and (pnRebase) are non-NULL and
 ** the input is a changeset (not a patchset), then sqlite3changeset_apply_v2()
-** may set (*ppRebase) to point to a "rebase" that may be used with the 
+** may set (*ppRebase) to point to a "rebase" that may be used with the
 ** sqlite3_rebaser APIs buffer before returning. In this case (*pnRebase)
 ** is set to the size of the buffer in bytes. It is the responsibility of the
 ** caller to eventually free any such buffer using sqlite3_free(). The buffer
@@ -11020,18 +12334,39 @@ SQLITE_API int sqlite3changeset_apply_v2(
 **   SAVEPOINT is committed if the changeset or patchset is successfully
 **   applied, or rolled back if an error occurs. Specifying this flag
 **   causes the sessions module to omit this savepoint. In this case, if the
-**   caller has an open transaction or savepoint when apply_v2() is called, 
+**   caller has an open transaction or savepoint when apply_v2() is called,
 **   it may revert the partially applied changeset by rolling it back.
 **
 ** <dt>SQLITE_CHANGESETAPPLY_INVERT <dd>
 **   Invert the changeset before applying it. This is equivalent to inverting
 **   a changeset using sqlite3changeset_invert() before applying it. It is
 **   an error to specify this flag with a patchset.
+**
+** <dt>SQLITE_CHANGESETAPPLY_IGNORENOOP <dd>
+**   Do not invoke the conflict handler callback for any changes that
+**   would not actually modify the database even if they were applied.
+**   Specifically, this means that the conflict handler is not invoked
+**   for:
+**    <ul>
+**    <li>a delete change if the row being deleted cannot be found,
+**    <li>an update change if the modified fields are already set to
+**        their new values in the conflicting row, or
+**    <li>an insert change if all fields of the conflicting row match
+**        the row being inserted.
+**    </ul>
+**
+** <dt>SQLITE_CHANGESETAPPLY_FKNOACTION <dd>
+**   If this flag it set, then all foreign key constraints in the target
+**   database behave as if they were declared with "ON UPDATE NO ACTION ON
+**   DELETE NO ACTION", even if they are actually CASCADE, RESTRICT, SET NULL
+**   or SET DEFAULT.
 */
 #define SQLITE_CHANGESETAPPLY_NOSAVEPOINT   0x0001
 #define SQLITE_CHANGESETAPPLY_INVERT        0x0002
+#define SQLITE_CHANGESETAPPLY_IGNORENOOP    0x0004
+#define SQLITE_CHANGESETAPPLY_FKNOACTION    0x0008
 
-/* 
+/*
 ** CAPI3REF: Constants Passed To The Conflict Handler
 **
 ** Values that may be passed as the second argument to a conflict-handler.
@@ -11040,32 +12375,32 @@ SQLITE_API int sqlite3changeset_apply_v2(
 ** <dt>SQLITE_CHANGESET_DATA<dd>
 **   The conflict handler is invoked with CHANGESET_DATA as the second argument
 **   when processing a DELETE or UPDATE change if a row with the required
-**   PRIMARY KEY fields is present in the database, but one or more other 
-**   (non primary-key) fields modified by the update do not contain the 
+**   PRIMARY KEY fields is present in the database, but one or more other
+**   (non primary-key) fields modified by the update do not contain the
 **   expected "before" values.
-** 
+**
 **   The conflicting row, in this case, is the database row with the matching
 **   primary key.
-** 
+**
 ** <dt>SQLITE_CHANGESET_NOTFOUND<dd>
 **   The conflict handler is invoked with CHANGESET_NOTFOUND as the second
 **   argument when processing a DELETE or UPDATE change if a row with the
 **   required PRIMARY KEY fields is not present in the database.
-** 
+**
 **   There is no conflicting row in this case. The results of invoking the
 **   sqlite3changeset_conflict() API are undefined.
-** 
+**
 ** <dt>SQLITE_CHANGESET_CONFLICT<dd>
 **   CHANGESET_CONFLICT is passed as the second argument to the conflict
-**   handler while processing an INSERT change if the operation would result 
+**   handler while processing an INSERT change if the operation would result
 **   in duplicate primary key values.
-** 
+**
 **   The conflicting row in this case is the database row with the matching
 **   primary key.
 **
 ** <dt>SQLITE_CHANGESET_FOREIGN_KEY<dd>
 **   If foreign key handling is enabled, and applying a changeset leaves the
-**   database in a state containing foreign key violations, the conflict 
+**   database in a state containing foreign key violations, the conflict
 **   handler is invoked with CHANGESET_FOREIGN_KEY as the second argument
 **   exactly once before the changeset is committed. If the conflict handler
 **   returns CHANGESET_OMIT, the changes, including those that caused the
@@ -11075,12 +12410,12 @@ SQLITE_API int sqlite3changeset_apply_v2(
 **   No current or conflicting row information is provided. The only function
 **   it is possible to call on the supplied sqlite3_changeset_iter handle
 **   is sqlite3changeset_fk_conflicts().
-** 
+**
 ** <dt>SQLITE_CHANGESET_CONSTRAINT<dd>
-**   If any other constraint violation occurs while applying a change (i.e. 
-**   a UNIQUE, CHECK or NOT NULL constraint), the conflict handler is 
+**   If any other constraint violation occurs while applying a change (i.e.
+**   a UNIQUE, CHECK or NOT NULL constraint), the conflict handler is
 **   invoked with CHANGESET_CONSTRAINT as the second argument.
-** 
+**
 **   There is no conflicting row in this case. The results of invoking the
 **   sqlite3changeset_conflict() API are undefined.
 **
@@ -11092,7 +12427,7 @@ SQLITE_API int sqlite3changeset_apply_v2(
 #define SQLITE_CHANGESET_CONSTRAINT  4
 #define SQLITE_CHANGESET_FOREIGN_KEY 5
 
-/* 
+/*
 ** CAPI3REF: Constants Returned By The Conflict Handler
 **
 ** A conflict handler callback must return one of the following three values.
@@ -11100,13 +12435,13 @@ SQLITE_API int sqlite3changeset_apply_v2(
 ** <dl>
 ** <dt>SQLITE_CHANGESET_OMIT<dd>
 **   If a conflict handler returns this value no special action is taken. The
-**   change that caused the conflict is not applied. The session module 
+**   change that caused the conflict is not applied. The session module
 **   continues to the next change in the changeset.
 **
 ** <dt>SQLITE_CHANGESET_REPLACE<dd>
 **   This value may only be returned if the second argument to the conflict
 **   handler was SQLITE_CHANGESET_DATA or SQLITE_CHANGESET_CONFLICT. If this
-**   is not the case, any changes applied so far are rolled back and the 
+**   is not the case, any changes applied so far are rolled back and the
 **   call to sqlite3changeset_apply() returns SQLITE_MISUSE.
 **
 **   If CHANGESET_REPLACE is returned by an SQLITE_CHANGESET_DATA conflict
@@ -11119,7 +12454,7 @@ SQLITE_API int sqlite3changeset_apply_v2(
 **   the original row is restored to the database before continuing.
 **
 ** <dt>SQLITE_CHANGESET_ABORT<dd>
-**   If this value is returned, any changes applied so far are rolled back 
+**   If this value is returned, any changes applied so far are rolled back
 **   and the call to sqlite3changeset_apply() returns SQLITE_ABORT.
 ** </dl>
 */
@@ -11127,20 +12462,20 @@ SQLITE_API int sqlite3changeset_apply_v2(
 #define SQLITE_CHANGESET_REPLACE    1
 #define SQLITE_CHANGESET_ABORT      2
 
-/* 
+/*
 ** CAPI3REF: Rebasing changesets
 ** EXPERIMENTAL
 **
 ** Suppose there is a site hosting a database in state S0. And that
 ** modifications are made that move that database to state S1 and a
 ** changeset recorded (the "local" changeset). Then, a changeset based
-** on S0 is received from another site (the "remote" changeset) and 
-** applied to the database. The database is then in state 
+** on S0 is received from another site (the "remote" changeset) and
+** applied to the database. The database is then in state
 ** (S1+"remote"), where the exact state depends on any conflict
 ** resolution decisions (OMIT or REPLACE) made while applying "remote".
-** Rebasing a changeset is to update it to take those conflict 
+** Rebasing a changeset is to update it to take those conflict
 ** resolution decisions into account, so that the same conflicts
-** do not have to be resolved elsewhere in the network. 
+** do not have to be resolved elsewhere in the network.
 **
 ** For example, if both the local and remote changesets contain an
 ** INSERT of the same key on "CREATE TABLE t1(a PRIMARY KEY, b)":
@@ -11159,7 +12494,7 @@ SQLITE_API int sqlite3changeset_apply_v2(
 **
 ** <dl>
 ** <dt>Local INSERT<dd>
-**   This may only conflict with a remote INSERT. If the conflict 
+**   This may only conflict with a remote INSERT. If the conflict
 **   resolution was OMIT, then add an UPDATE change to the rebased
 **   changeset. Or, if the conflict resolution was REPLACE, add
 **   nothing to the rebased changeset.
@@ -11183,12 +12518,12 @@ SQLITE_API int sqlite3changeset_apply_v2(
 **   the old.* values are rebased using the new.* values in the remote
 **   change. Or, if the resolution is REPLACE, then the change is copied
 **   into the rebased changeset with updates to columns also updated by
-**   the conflicting remote UPDATE removed. If this means no columns would 
+**   the conflicting remote UPDATE removed. If this means no columns would
 **   be updated, the change is omitted.
 ** </dl>
 **
-** A local change may be rebased against multiple remote changes 
-** simultaneously. If a single key is modified by multiple remote 
+** A local change may be rebased against multiple remote changes
+** simultaneously. If a single key is modified by multiple remote
 ** changesets, they are combined as follows before the local changeset
 ** is rebased:
 **
@@ -11201,10 +12536,10 @@ SQLITE_API int sqlite3changeset_apply_v2(
 **         of the OMIT resolutions.
 ** </ul>
 **
-** Note that conflict resolutions from multiple remote changesets are 
-** combined on a per-field basis, not per-row. This means that in the 
-** case of multiple remote UPDATE operations, some fields of a single 
-** local change may be rebased for REPLACE while others are rebased for 
+** Note that conflict resolutions from multiple remote changesets are
+** combined on a per-field basis, not per-row. This means that in the
+** case of multiple remote UPDATE operations, some fields of a single
+** local change may be rebased for REPLACE while others are rebased for
 ** OMIT.
 **
 ** In order to rebase a local changeset, the remote changeset must first
@@ -11212,7 +12547,7 @@ SQLITE_API int sqlite3changeset_apply_v2(
 ** the buffer of rebase information captured. Then:
 **
 ** <ol>
-**   <li> An sqlite3_rebaser object is created by calling 
+**   <li> An sqlite3_rebaser object is created by calling
 **        sqlite3rebaser_create().
 **   <li> The new object is configured with the rebase buffer obtained from
 **        sqlite3changeset_apply_v2() by calling sqlite3rebaser_configure().
@@ -11233,8 +12568,8 @@ typedef struct sqlite3_rebaser sqlite3_rebaser;
 **
 ** Allocate a new changeset rebaser object. If successful, set (*ppNew) to
 ** point to the new object and return SQLITE_OK. Otherwise, if an error
-** occurs, return an SQLite error code (e.g. SQLITE_NOMEM) and set (*ppNew) 
-** to NULL. 
+** occurs, return an SQLite error code (e.g. SQLITE_NOMEM) and set (*ppNew)
+** to NULL.
 */
 SQLITE_API int sqlite3rebaser_create(sqlite3_rebaser **ppNew);
 
@@ -11248,9 +12583,9 @@ SQLITE_API int sqlite3rebaser_create(sqlite3_rebaser **ppNew);
 ** sqlite3changeset_apply_v2().
 */
 SQLITE_API int sqlite3rebaser_configure(
-  sqlite3_rebaser*, 
+  sqlite3_rebaser*,
   int nRebase, const void *pRebase
-); 
+);
 
 /*
 ** CAPI3REF: Rebase a changeset
@@ -11260,7 +12595,7 @@ SQLITE_API int sqlite3rebaser_configure(
 ** in size. This function allocates and populates a buffer with a copy
 ** of the changeset rebased according to the configuration of the
 ** rebaser object passed as the first argument. If successful, (*ppOut)
-** is set to point to the new buffer containing the rebased changeset and 
+** is set to point to the new buffer containing the rebased changeset and
 ** (*pnOut) to its size in bytes and SQLITE_OK returned. It is the
 ** responsibility of the caller to eventually free the new buffer using
 ** sqlite3_free(). Otherwise, if an error occurs, (*ppOut) and (*pnOut)
@@ -11268,8 +12603,8 @@ SQLITE_API int sqlite3rebaser_configure(
 */
 SQLITE_API int sqlite3rebaser_rebase(
   sqlite3_rebaser*,
-  int nIn, const void *pIn, 
-  int *pnOut, void **ppOut 
+  int nIn, const void *pIn,
+  int *pnOut, void **ppOut
 );
 
 /*
@@ -11280,30 +12615,30 @@ SQLITE_API int sqlite3rebaser_rebase(
 ** should be one call to this function for each successful invocation
 ** of sqlite3rebaser_create().
 */
-SQLITE_API void sqlite3rebaser_delete(sqlite3_rebaser *p); 
+SQLITE_API void sqlite3rebaser_delete(sqlite3_rebaser *p);
 
 /*
 ** CAPI3REF: Streaming Versions of API functions.
 **
-** The six streaming API xxx_strm() functions serve similar purposes to the 
+** The six streaming API xxx_strm() functions serve similar purposes to the
 ** corresponding non-streaming API functions:
 **
 ** <table border=1 style="margin-left:8ex;margin-right:8ex">
 **   <tr><th>Streaming function<th>Non-streaming equivalent</th>
-**   <tr><td>sqlite3changeset_apply_strm<td>[sqlite3changeset_apply] 
-**   <tr><td>sqlite3changeset_apply_strm_v2<td>[sqlite3changeset_apply_v2] 
-**   <tr><td>sqlite3changeset_concat_strm<td>[sqlite3changeset_concat] 
-**   <tr><td>sqlite3changeset_invert_strm<td>[sqlite3changeset_invert] 
-**   <tr><td>sqlite3changeset_start_strm<td>[sqlite3changeset_start] 
-**   <tr><td>sqlite3session_changeset_strm<td>[sqlite3session_changeset] 
-**   <tr><td>sqlite3session_patchset_strm<td>[sqlite3session_patchset] 
+**   <tr><td>sqlite3changeset_apply_strm<td>[sqlite3changeset_apply]
+**   <tr><td>sqlite3changeset_apply_strm_v2<td>[sqlite3changeset_apply_v2]
+**   <tr><td>sqlite3changeset_concat_strm<td>[sqlite3changeset_concat]
+**   <tr><td>sqlite3changeset_invert_strm<td>[sqlite3changeset_invert]
+**   <tr><td>sqlite3changeset_start_strm<td>[sqlite3changeset_start]
+**   <tr><td>sqlite3session_changeset_strm<td>[sqlite3session_changeset]
+**   <tr><td>sqlite3session_patchset_strm<td>[sqlite3session_patchset]
 ** </table>
 **
 ** Non-streaming functions that accept changesets (or patchsets) as input
-** require that the entire changeset be stored in a single buffer in memory. 
-** Similarly, those that return a changeset or patchset do so by returning 
-** a pointer to a single large buffer allocated using sqlite3_malloc(). 
-** Normally this is convenient. However, if an application running in a 
+** require that the entire changeset be stored in a single buffer in memory.
+** Similarly, those that return a changeset or patchset do so by returning
+** a pointer to a single large buffer allocated using sqlite3_malloc().
+** Normally this is convenient. However, if an application running in a
 ** low-memory environment is required to handle very large changesets, the
 ** large contiguous memory allocations required can become onerous.
 **
@@ -11325,12 +12660,12 @@ SQLITE_API void sqlite3rebaser_delete(sqlite3_rebaser *p);
 **  </pre>
 **
 ** Each time the xInput callback is invoked by the sessions module, the first
-** argument passed is a copy of the supplied pIn context pointer. The second 
-** argument, pData, points to a buffer (*pnData) bytes in size. Assuming no 
-** error occurs the xInput method should copy up to (*pnData) bytes of data 
-** into the buffer and set (*pnData) to the actual number of bytes copied 
-** before returning SQLITE_OK. If the input is completely exhausted, (*pnData) 
-** should be set to zero to indicate this. Or, if an error occurs, an SQLite 
+** argument passed is a copy of the supplied pIn context pointer. The second
+** argument, pData, points to a buffer (*pnData) bytes in size. Assuming no
+** error occurs the xInput method should copy up to (*pnData) bytes of data
+** into the buffer and set (*pnData) to the actual number of bytes copied
+** before returning SQLITE_OK. If the input is completely exhausted, (*pnData)
+** should be set to zero to indicate this. Or, if an error occurs, an SQLite
 ** error code should be returned. In all cases, if an xInput callback returns
 ** an error, all processing is abandoned and the streaming API function
 ** returns a copy of the error code to the caller.
@@ -11338,7 +12673,7 @@ SQLITE_API void sqlite3rebaser_delete(sqlite3_rebaser *p);
 ** In the case of sqlite3changeset_start_strm(), the xInput callback may be
 ** invoked by the sessions module at any point during the lifetime of the
 ** iterator. If such an xInput callback returns an error, the iterator enters
-** an error state, whereby all subsequent calls to iterator functions 
+** an error state, whereby all subsequent calls to iterator functions
 ** immediately fail with the same error code as returned by xInput.
 **
 ** Similarly, streaming API functions that return changesets (or patchsets)
@@ -11368,7 +12703,7 @@ SQLITE_API void sqlite3rebaser_delete(sqlite3_rebaser *p);
 ** is immediately abandoned and the streaming API function returns a copy
 ** of the xOutput error code to the application.
 **
-** The sessions module never invokes an xOutput callback with the third 
+** The sessions module never invokes an xOutput callback with the third
 ** parameter set to a value less than or equal to zero. Other than this,
 ** no guarantees are made as to the size of the chunks of data returned.
 */
@@ -11439,12 +12774,12 @@ SQLITE_API int sqlite3session_patchset_strm(
   int (*xOutput)(void *pOut, const void *pData, int nData),
   void *pOut
 );
-SQLITE_API int sqlite3changegroup_add_strm(sqlite3_changegroup*, 
+SQLITE_API int sqlite3changegroup_add_strm(sqlite3_changegroup*,
     int (*xInput)(void *pIn, void *pData, int *pnData),
     void *pIn
 );
 SQLITE_API int sqlite3changegroup_output_strm(sqlite3_changegroup*,
-    int (*xOutput)(void *pOut, const void *pData, int nData), 
+    int (*xOutput)(void *pOut, const void *pData, int nData),
     void *pOut
 );
 SQLITE_API int sqlite3rebaser_rebase_strm(
@@ -11459,16 +12794,16 @@ SQLITE_API int sqlite3rebaser_rebase_strm(
 ** CAPI3REF: Configure global parameters
 **
 ** The sqlite3session_config() interface is used to make global configuration
-** changes to the sessions module in order to tune it to the specific needs 
+** changes to the sessions module in order to tune it to the specific needs
 ** of the application.
 **
 ** The sqlite3session_config() interface is not threadsafe. If it is invoked
 ** while any other thread is inside any other sessions method then the
 ** results are undefined. Furthermore, if it is invoked after any sessions
-** related objects have been created, the results are also undefined. 
+** related objects have been created, the results are also undefined.
 **
 ** The first argument to the sqlite3session_config() function must be one
-** of the SQLITE_SESSION_CONFIG_XXX constants defined below. The 
+** of the SQLITE_SESSION_CONFIG_XXX constants defined below. The
 ** interpretation of the (void*) value passed as the second parameter and
 ** the effect of calling this function depends on the value of the first
 ** parameter.
@@ -11518,7 +12853,7 @@ SQLITE_API int sqlite3session_config(int op, void *pArg);
 **
 ******************************************************************************
 **
-** Interfaces to extend FTS5. Using the interfaces defined in this file, 
+** Interfaces to extend FTS5. Using the interfaces defined in this file,
 ** FTS5 may be extended with:
 **
 **     * custom tokenizers, and
@@ -11562,19 +12897,19 @@ struct Fts5PhraseIter {
 ** EXTENSION API FUNCTIONS
 **
 ** xUserData(pFts):
-**   Return a copy of the context pointer the extension function was 
-**   registered with.
+**   Return a copy of the pUserData pointer passed to the xCreateFunction()
+**   API when the extension function was registered.
 **
 ** xColumnTotalSize(pFts, iCol, pnToken):
 **   If parameter iCol is less than zero, set output variable *pnToken
 **   to the total number of tokens in the FTS5 table. Or, if iCol is
 **   non-negative but less than the number of columns in the table, return
-**   the total number of tokens in column iCol, considering all rows in 
+**   the total number of tokens in column iCol, considering all rows in
 **   the FTS5 table.
 **
 **   If parameter iCol is greater than or equal to the number of columns
 **   in the table, SQLITE_RANGE is returned. Or, if an error occurs (e.g.
-**   an OOM condition or IO error), an appropriate SQLite error code is 
+**   an OOM condition or IO error), an appropriate SQLite error code is
 **   returned.
 **
 ** xColumnCount(pFts):
@@ -11588,15 +12923,18 @@ struct Fts5PhraseIter {
 **
 **   If parameter iCol is greater than or equal to the number of columns
 **   in the table, SQLITE_RANGE is returned. Or, if an error occurs (e.g.
-**   an OOM condition or IO error), an appropriate SQLite error code is 
+**   an OOM condition or IO error), an appropriate SQLite error code is
 **   returned.
 **
 **   This function may be quite inefficient if used with an FTS5 table
 **   created with the "columnsize=0" option.
 **
 ** xColumnText:
-**   This function attempts to retrieve the text of column iCol of the
-**   current document. If successful, (*pz) is set to point to a buffer
+**   If parameter iCol is less than zero, or greater than or equal to the
+**   number of columns in the table, SQLITE_RANGE is returned.
+**
+**   Otherwise, this function attempts to retrieve the text of column iCol of
+**   the current document. If successful, (*pz) is set to point to a buffer
 **   containing the text in utf-8 encoding, (*pn) is set to the size in bytes
 **   (not characters) of the buffer and SQLITE_OK is returned. Otherwise,
 **   if an error occurs, an SQLite error code is returned and the final values
@@ -11606,8 +12944,10 @@ struct Fts5PhraseIter {
 **   Returns the number of phrases in the current query expression.
 **
 ** xPhraseSize:
-**   Returns the number of tokens in phrase iPhrase of the query. Phrases
-**   are numbered starting from zero.
+**   If parameter iCol is less than zero, or greater than or equal to the
+**   number of phrases in the current query, as returned by xPhraseCount,
+**   0 is returned. Otherwise, this function returns the number of tokens in
+**   phrase iPhrase of the query. Phrases are numbered starting from zero.
 **
 ** xInstCount:
 **   Set *pnInst to the total number of occurrences of all phrases within
@@ -11615,23 +12955,24 @@ struct Fts5PhraseIter {
 **   an error code (i.e. SQLITE_NOMEM) if an error occurs.
 **
 **   This API can be quite slow if used with an FTS5 table created with the
-**   "detail=none" or "detail=column" option. If the FTS5 table is created 
-**   with either "detail=none" or "detail=column" and "content=" option 
+**   "detail=none" or "detail=column" option. If the FTS5 table is created
+**   with either "detail=none" or "detail=column" and "content=" option
 **   (i.e. if it is a contentless table), then this API always returns 0.
 **
 ** xInst:
 **   Query for the details of phrase match iIdx within the current row.
 **   Phrase matches are numbered starting from zero, so the iIdx argument
 **   should be greater than or equal to zero and smaller than the value
-**   output by xInstCount().
+**   output by xInstCount(). If iIdx is less than zero or greater than
+**   or equal to the value returned by xInstCount(), SQLITE_RANGE is returned.
 **
-**   Usually, output parameter *piPhrase is set to the phrase number, *piCol
+**   Otherwise, output parameter *piPhrase is set to the phrase number, *piCol
 **   to the column in which it occurs and *piOff the token offset of the
-**   first token of the phrase. Returns SQLITE_OK if successful, or an error
-**   code (i.e. SQLITE_NOMEM) if an error occurs.
+**   first token of the phrase. SQLITE_OK is returned if successful, or an
+**   error code (i.e. SQLITE_NOMEM) if an error occurs.
 **
 **   This API can be quite slow if used with an FTS5 table created with the
-**   "detail=none" or "detail=column" option. 
+**   "detail=none" or "detail=column" option.
 **
 ** xRowid:
 **   Returns the rowid of the current row.
@@ -11647,13 +12988,17 @@ struct Fts5PhraseIter {
 **
 **   with $p set to a phrase equivalent to the phrase iPhrase of the
 **   current query is executed. Any column filter that applies to
-**   phrase iPhrase of the current query is included in $p. For each 
-**   row visited, the callback function passed as the fourth argument 
-**   is invoked. The context and API objects passed to the callback 
+**   phrase iPhrase of the current query is included in $p. For each
+**   row visited, the callback function passed as the fourth argument
+**   is invoked. The context and API objects passed to the callback
 **   function may be used to access the properties of each matched row.
-**   Invoking Api.xUserData() returns a copy of the pointer passed as 
+**   Invoking Api.xUserData() returns a copy of the pointer passed as
 **   the third argument to pUserData.
 **
+**   If parameter iPhrase is less than zero, or greater than or equal to
+**   the number of phrases in the query, as returned by xPhraseCount(),
+**   this function returns SQLITE_RANGE.
+**
 **   If the callback function returns any value other than SQLITE_OK, the
 **   query is abandoned and the xQueryPhrase function returns immediately.
 **   If the returned value is SQLITE_DONE, xQueryPhrase returns SQLITE_OK.
@@ -11666,14 +13011,14 @@ struct Fts5PhraseIter {
 **
 ** xSetAuxdata(pFts5, pAux, xDelete)
 **
-**   Save the pointer passed as the second argument as the extension function's 
+**   Save the pointer passed as the second argument as the extension function's
 **   "auxiliary data". The pointer may then be retrieved by the current or any
 **   future invocation of the same fts5 extension function made as part of
 **   the same MATCH query using the xGetAuxdata() API.
 **
 **   Each extension function is allocated a single auxiliary data slot for
-**   each FTS query (MATCH expression). If the extension function is invoked 
-**   more than once for a single FTS query, then all invocations share a 
+**   each FTS query (MATCH expression). If the extension function is invoked
+**   more than once for a single FTS query, then all invocations share a
 **   single auxiliary data context.
 **
 **   If there is already an auxiliary data pointer when this function is
@@ -11692,7 +13037,7 @@ struct Fts5PhraseIter {
 **
 ** xGetAuxdata(pFts5, bClear)
 **
-**   Returns the current auxiliary data pointer for the fts5 extension 
+**   Returns the current auxiliary data pointer for the fts5 extension
 **   function. See the xSetAuxdata() method for details.
 **
 **   If the bClear argument is non-zero, then the auxiliary data is cleared
@@ -11712,7 +13057,7 @@ struct Fts5PhraseIter {
 **   method, to iterate through all instances of a single query phrase within
 **   the current row. This is the same information as is accessible via the
 **   xInstCount/xInst APIs. While the xInstCount/xInst APIs are more convenient
-**   to use, this API may be faster under some circumstances. To iterate 
+**   to use, this API may be faster under some circumstances. To iterate
 **   through instances of phrase iPhrase, use the following code:
 **
 **       Fts5PhraseIter iter;
@@ -11730,11 +13075,15 @@ struct Fts5PhraseIter {
 **   xPhraseFirstColumn() and xPhraseNextColumn() as illustrated below).
 **
 **   This API can be quite slow if used with an FTS5 table created with the
-**   "detail=none" or "detail=column" option. If the FTS5 table is created 
-**   with either "detail=none" or "detail=column" and "content=" option 
+**   "detail=none" or "detail=column" option. If the FTS5 table is created
+**   with either "detail=none" or "detail=column" and "content=" option
 **   (i.e. if it is a contentless table), then this API always iterates
 **   through an empty set (all calls to xPhraseFirst() set iCol to -1).
 **
+**   In all cases, matches are visited in (column ASC, offset ASC) order.
+**   i.e. all those in column 0, sorted by offset, followed by those in
+**   column 1, etc.
+**
 ** xPhraseNext()
 **   See xPhraseFirst above.
 **
@@ -11755,22 +13104,78 @@ struct Fts5PhraseIter {
 **       }
 **
 **   This API can be quite slow if used with an FTS5 table created with the
-**   "detail=none" option. If the FTS5 table is created with either 
-**   "detail=none" "content=" option (i.e. if it is a contentless table), 
-**   then this API always iterates through an empty set (all calls to 
+**   "detail=none" option. If the FTS5 table is created with either
+**   "detail=none" "content=" option (i.e. if it is a contentless table),
+**   then this API always iterates through an empty set (all calls to
 **   xPhraseFirstColumn() set iCol to -1).
 **
 **   The information accessed using this API and its companion
 **   xPhraseFirstColumn() may also be obtained using xPhraseFirst/xPhraseNext
 **   (or xInst/xInstCount). The chief advantage of this API is that it is
 **   significantly more efficient than those alternatives when used with
-**   "detail=column" tables.  
+**   "detail=column" tables.
 **
 ** xPhraseNextColumn()
 **   See xPhraseFirstColumn above.
+**
+** xQueryToken(pFts5, iPhrase, iToken, ppToken, pnToken)
+**   This is used to access token iToken of phrase iPhrase of the current
+**   query. Before returning, output parameter *ppToken is set to point
+**   to a buffer containing the requested token, and *pnToken to the
+**   size of this buffer in bytes.
+**
+**   If iPhrase or iToken are less than zero, or if iPhrase is greater than
+**   or equal to the number of phrases in the query as reported by
+**   xPhraseCount(), or if iToken is equal to or greater than the number of
+**   tokens in the phrase, SQLITE_RANGE is returned and *ppToken and *pnToken
+     are both zeroed.
+**
+**   The output text is not a copy of the query text that specified the
+**   token. It is the output of the tokenizer module. For tokendata=1
+**   tables, this includes any embedded 0x00 and trailing data.
+**
+** xInstToken(pFts5, iIdx, iToken, ppToken, pnToken)
+**   This is used to access token iToken of phrase hit iIdx within the
+**   current row. If iIdx is less than zero or greater than or equal to the
+**   value returned by xInstCount(), SQLITE_RANGE is returned.  Otherwise,
+**   output variable (*ppToken) is set to point to a buffer containing the
+**   matching document token, and (*pnToken) to the size of that buffer in
+**   bytes. This API is not available if the specified token matches a
+**   prefix query term. In that case both output variables are always set
+**   to 0.
+**
+**   The output text is not a copy of the document text that was tokenized.
+**   It is the output of the tokenizer module. For tokendata=1 tables, this
+**   includes any embedded 0x00 and trailing data.
+**
+**   This API can be quite slow if used with an FTS5 table created with the
+**   "detail=none" or "detail=column" option.
+**
+** xColumnLocale(pFts5, iIdx, pzLocale, pnLocale)
+**   If parameter iCol is less than zero, or greater than or equal to the
+**   number of columns in the table, SQLITE_RANGE is returned.
+**
+**   Otherwise, this function attempts to retrieve the locale associated
+**   with column iCol of the current row. Usually, there is no associated
+**   locale, and output parameters (*pzLocale) and (*pnLocale) are set
+**   to NULL and 0, respectively. However, if the fts5_locale() function
+**   was used to associate a locale with the value when it was inserted
+**   into the fts5 table, then (*pzLocale) is set to point to a nul-terminated
+**   buffer containing the name of the locale in utf-8 encoding. (*pnLocale)
+**   is set to the size in bytes of the buffer, not including the
+**   nul-terminator.
+**
+**   If successful, SQLITE_OK is returned. Or, if an error occurs, an
+**   SQLite error code is returned. The final value of the output parameters
+**   is undefined in this case.
+**
+** xTokenize_v2:
+**   Tokenize text using the tokenizer belonging to the FTS5 table. This
+**   API is the same as the xTokenize() API, except that it allows a tokenizer
+**   locale to be specified.
 */
 struct Fts5ExtensionApi {
-  int iVersion;                   /* Currently always set to 3 */
+  int iVersion;                   /* Currently always set to 4 */
 
   void *(*xUserData)(Fts5Context*);
 
@@ -11778,7 +13183,7 @@ struct Fts5ExtensionApi {
   int (*xRowCount)(Fts5Context*, sqlite3_int64 *pnRow);
   int (*xColumnTotalSize)(Fts5Context*, int iCol, sqlite3_int64 *pnToken);
 
-  int (*xTokenize)(Fts5Context*, 
+  int (*xTokenize)(Fts5Context*,
     const char *pText, int nText, /* Text to tokenize */
     void *pCtx,                   /* Context passed to xToken() */
     int (*xToken)(void*, int, const char*, int, int, int)       /* Callback */
@@ -11805,17 +13210,33 @@ struct Fts5ExtensionApi {
 
   int (*xPhraseFirstColumn)(Fts5Context*, int iPhrase, Fts5PhraseIter*, int*);
   void (*xPhraseNextColumn)(Fts5Context*, Fts5PhraseIter*, int *piCol);
+
+  /* Below this point are iVersion>=3 only */
+  int (*xQueryToken)(Fts5Context*,
+      int iPhrase, int iToken,
+      const char **ppToken, int *pnToken
+  );
+  int (*xInstToken)(Fts5Context*, int iIdx, int iToken, const char**, int*);
+
+  /* Below this point are iVersion>=4 only */
+  int (*xColumnLocale)(Fts5Context*, int iCol, const char **pz, int *pn);
+  int (*xTokenize_v2)(Fts5Context*,
+    const char *pText, int nText,      /* Text to tokenize */
+    const char *pLocale, int nLocale,  /* Locale to pass to tokenizer */
+    void *pCtx,                        /* Context passed to xToken() */
+    int (*xToken)(void*, int, const char*, int, int, int)       /* Callback */
+  );
 };
 
-/* 
+/*
 ** CUSTOM AUXILIARY FUNCTIONS
 *************************************************************************/
 
 /*************************************************************************
 ** CUSTOM TOKENIZERS
 **
-** Applications may also register custom tokenizer types. A tokenizer 
-** is registered by providing fts5 with a populated instance of the 
+** Applications may also register custom tokenizer types. A tokenizer
+** is registered by providing fts5 with a populated instance of the
 ** following structure. All structure methods must be defined, setting
 ** any member of the fts5_tokenizer struct to NULL leads to undefined
 ** behaviour. The structure methods are expected to function as follows:
@@ -11825,17 +13246,17 @@ struct Fts5ExtensionApi {
 **   A tokenizer instance is required to actually tokenize text.
 **
 **   The first argument passed to this function is a copy of the (void*)
-**   pointer provided by the application when the fts5_tokenizer object
-**   was registered with FTS5 (the third argument to xCreateTokenizer()). 
+**   pointer provided by the application when the fts5_tokenizer_v2 object
+**   was registered with FTS5 (the third argument to xCreateTokenizer()).
 **   The second and third arguments are an array of nul-terminated strings
 **   containing the tokenizer arguments, if any, specified following the
 **   tokenizer name as part of the CREATE VIRTUAL TABLE statement used
 **   to create the FTS5 table.
 **
-**   The final argument is an output variable. If successful, (*ppOut) 
+**   The final argument is an output variable. If successful, (*ppOut)
 **   should be set to point to the new tokenizer handle and SQLITE_OK
 **   returned. If an error occurs, some value other than SQLITE_OK should
-**   be returned. In this case, fts5 assumes that the final value of *ppOut 
+**   be returned. In this case, fts5 assumes that the final value of *ppOut
 **   is undefined.
 **
 ** xDelete:
@@ -11844,12 +13265,12 @@ struct Fts5ExtensionApi {
 **   be invoked exactly once for each successful call to xCreate().
 **
 ** xTokenize:
-**   This function is expected to tokenize the nText byte string indicated 
+**   This function is expected to tokenize the nText byte string indicated
 **   by argument pText. pText may or may not be nul-terminated. The first
 **   argument passed to this function is a pointer to an Fts5Tokenizer object
 **   returned by an earlier call to xCreate().
 **
-**   The second argument indicates the reason that FTS5 is requesting
+**   The third argument indicates the reason that FTS5 is requesting
 **   tokenization of the supplied text. This is always one of the following
 **   four values:
 **
@@ -11858,8 +13279,8 @@ struct Fts5ExtensionApi {
 **            determine the set of tokens to add to (or delete from) the
 **            FTS index.
 **
-**       <li> <b>FTS5_TOKENIZE_QUERY</b> - A MATCH query is being executed 
-**            against the FTS index. The tokenizer is being called to tokenize 
+**       <li> <b>FTS5_TOKENIZE_QUERY</b> - A MATCH query is being executed
+**            against the FTS index. The tokenizer is being called to tokenize
 **            a bareword or quoted string specified as part of the query.
 **
 **       <li> <b>(FTS5_TOKENIZE_QUERY | FTS5_TOKENIZE_PREFIX)</b> - Same as
@@ -11867,12 +13288,19 @@ struct Fts5ExtensionApi {
 **            followed by a "*" character, indicating that the last token
 **            returned by the tokenizer will be treated as a token prefix.
 **
-**       <li> <b>FTS5_TOKENIZE_AUX</b> - The tokenizer is being invoked to 
+**       <li> <b>FTS5_TOKENIZE_AUX</b> - The tokenizer is being invoked to
 **            satisfy an fts5_api.xTokenize() request made by an auxiliary
 **            function. Or an fts5_api.xColumnSize() request made by the same
-**            on a columnsize=0 database.  
+**            on a columnsize=0 database.
 **   </ul>
 **
+**   The sixth and seventh arguments passed to xTokenize() - pLocale and
+**   nLocale - are a pointer to a buffer containing the locale to use for
+**   tokenization (e.g. "en_US") and its size in bytes, respectively. The
+**   pLocale buffer is not nul-terminated. pLocale may be passed NULL (in
+**   which case nLocale is always 0) to indicate that the tokenizer should
+**   use its default locale.
+**
 **   For each token in the input string, the supplied callback xToken() must
 **   be invoked. The first argument to it should be a copy of the pointer
 **   passed as the second argument to xTokenize(). The third and fourth
@@ -11882,10 +13310,10 @@ struct Fts5ExtensionApi {
 **   which the token is derived within the input.
 **
 **   The second argument passed to the xToken() callback ("tflags") should
-**   normally be set to 0. The exception is if the tokenizer supports 
+**   normally be set to 0. The exception is if the tokenizer supports
 **   synonyms. In this case see the discussion below for details.
 **
-**   FTS5 assumes the xToken() callback is invoked for each token in the 
+**   FTS5 assumes the xToken() callback is invoked for each token in the
 **   order that they occur within the input text.
 **
 **   If an xToken() callback returns any value other than SQLITE_OK, then
@@ -11896,10 +13324,34 @@ struct Fts5ExtensionApi {
 **   may abandon the tokenization and return any error code other than
 **   SQLITE_OK or SQLITE_DONE.
 **
+**   If the tokenizer is registered using an fts5_tokenizer_v2 object,
+**   then the xTokenize() method has two additional arguments - pLocale
+**   and nLocale. These specify the locale that the tokenizer should use
+**   for the current request. If pLocale and nLocale are both 0, then the
+**   tokenizer should use its default locale. Otherwise, pLocale points to
+**   an nLocale byte buffer containing the name of the locale to use as utf-8
+**   text. pLocale is not nul-terminated.
+**
+** FTS5_TOKENIZER
+**
+** There is also an fts5_tokenizer object. This is an older, deprecated,
+** version of fts5_tokenizer_v2. It is similar except that:
+**
+**  <ul>
+**    <li> There is no "iVersion" field, and
+**    <li> The xTokenize() method does not take a locale argument.
+**  </ul>
+**
+** Legacy fts5_tokenizer tokenizers must be registered using the
+** legacy xCreateTokenizer() function, instead of xCreateTokenizer_v2().
+**
+** Tokenizer implementations registered using either API may be retrieved
+** using both xFindTokenizer() and xFindTokenizer_v2().
+**
 ** SYNONYM SUPPORT
 **
 **   Custom tokenizers may also support synonyms. Consider a case in which a
-**   user wishes to query for a phrase such as "first place". Using the 
+**   user wishes to query for a phrase such as "first place". Using the
 **   built-in tokenizers, the FTS5 query 'first + place' will match instances
 **   of "first place" within the document set, but not alternative forms
 **   such as "1st place". In some applications, it would be better to match
@@ -11919,34 +13371,34 @@ struct Fts5ExtensionApi {
 **
 **       <li> By querying the index for all synonyms of each query term
 **            separately. In this case, when tokenizing query text, the
-**            tokenizer may provide multiple synonyms for a single term 
-**            within the document. FTS5 then queries the index for each 
+**            tokenizer may provide multiple synonyms for a single term
+**            within the document. FTS5 then queries the index for each
 **            synonym individually. For example, faced with the query:
 **
 **   <codeblock>
 **     ... MATCH 'first place'</codeblock>
 **
 **            the tokenizer offers both "1st" and "first" as synonyms for the
-**            first token in the MATCH query and FTS5 effectively runs a query 
+**            first token in the MATCH query and FTS5 effectively runs a query
 **            similar to:
 **
 **   <codeblock>
 **     ... MATCH '(first OR 1st) place'</codeblock>
 **
 **            except that, for the purposes of auxiliary functions, the query
-**            still appears to contain just two phrases - "(first OR 1st)" 
+**            still appears to contain just two phrases - "(first OR 1st)"
 **            being treated as a single phrase.
 **
 **       <li> By adding multiple synonyms for a single term to the FTS index.
 **            Using this method, when tokenizing document text, the tokenizer
-**            provides multiple synonyms for each token. So that when a 
+**            provides multiple synonyms for each token. So that when a
 **            document such as "I won first place" is tokenized, entries are
 **            added to the FTS index for "i", "won", "first", "1st" and
 **            "place".
 **
 **            This way, even if the tokenizer does not provide synonyms
 **            when tokenizing query text (it should not - to do so would be
-**            inefficient), it doesn't matter if the user queries for 
+**            inefficient), it doesn't matter if the user queries for
 **            'first + place' or '1st + place', as there are entries in the
 **            FTS index corresponding to both forms of the first token.
 **   </ol>
@@ -11967,11 +13419,11 @@ struct Fts5ExtensionApi {
 **
 **   It is an error to specify the FTS5_TOKEN_COLOCATED flag the first time
 **   xToken() is called. Multiple synonyms may be specified for a single token
-**   by making multiple calls to xToken(FTS5_TOKEN_COLOCATED) in sequence. 
+**   by making multiple calls to xToken(FTS5_TOKEN_COLOCATED) in sequence.
 **   There is no limit to the number of synonyms that may be provided for a
 **   single token.
 **
-**   In many cases, method (1) above is the best approach. It does not add 
+**   In many cases, method (1) above is the best approach. It does not add
 **   extra data to the FTS index or require FTS5 to query for multiple terms,
 **   so it is efficient in terms of disk space and query speed. However, it
 **   does not support prefix queries very well. If, as suggested above, the
@@ -11983,35 +13435,62 @@ struct Fts5ExtensionApi {
 **   will not match documents that contain the token "1st" (as the tokenizer
 **   will probably not map "1s" to any prefix of "first").
 **
-**   For full prefix support, method (3) may be preferred. In this case, 
+**   For full prefix support, method (3) may be preferred. In this case,
 **   because the index contains entries for both "first" and "1st", prefix
 **   queries such as 'fi*' or '1s*' will match correctly. However, because
 **   extra entries are added to the FTS index, this method uses more space
 **   within the database.
 **
 **   Method (2) offers a midpoint between (1) and (3). Using this method,
-**   a query such as '1s*' will match documents that contain the literal 
+**   a query such as '1s*' will match documents that contain the literal
 **   token "1st", but not "first" (assuming the tokenizer is not able to
 **   provide synonyms for prefixes). However, a non-prefix query like '1st'
 **   will match against "1st" and "first". This method does not require
-**   extra disk space, as no extra entries are added to the FTS index. 
+**   extra disk space, as no extra entries are added to the FTS index.
 **   On the other hand, it may require more CPU cycles to run MATCH queries,
 **   as separate queries of the FTS index are required for each synonym.
 **
 **   When using methods (2) or (3), it is important that the tokenizer only
-**   provide synonyms when tokenizing document text (method (2)) or query
-**   text (method (3)), not both. Doing so will not cause any errors, but is
+**   provide synonyms when tokenizing document text (method (3)) or query
+**   text (method (2)), not both. Doing so will not cause any errors, but is
 **   inefficient.
 */
 typedef struct Fts5Tokenizer Fts5Tokenizer;
+typedef struct fts5_tokenizer_v2 fts5_tokenizer_v2;
+struct fts5_tokenizer_v2 {
+  int iVersion;             /* Currently always 2 */
+
+  int (*xCreate)(void*, const char **azArg, int nArg, Fts5Tokenizer **ppOut);
+  void (*xDelete)(Fts5Tokenizer*);
+  int (*xTokenize)(Fts5Tokenizer*,
+      void *pCtx,
+      int flags,            /* Mask of FTS5_TOKENIZE_* flags */
+      const char *pText, int nText,
+      const char *pLocale, int nLocale,
+      int (*xToken)(
+        void *pCtx,         /* Copy of 2nd argument to xTokenize() */
+        int tflags,         /* Mask of FTS5_TOKEN_* flags */
+        const char *pToken, /* Pointer to buffer containing token */
+        int nToken,         /* Size of token in bytes */
+        int iStart,         /* Byte offset of token within input text */
+        int iEnd            /* Byte offset of end of token within input text */
+      )
+  );
+};
+
+/*
+** New code should use the fts5_tokenizer_v2 type to define tokenizer
+** implementations. The following type is included for legacy applications
+** that still use it.
+*/
 typedef struct fts5_tokenizer fts5_tokenizer;
 struct fts5_tokenizer {
   int (*xCreate)(void*, const char **azArg, int nArg, Fts5Tokenizer **ppOut);
   void (*xDelete)(Fts5Tokenizer*);
-  int (*xTokenize)(Fts5Tokenizer*, 
+  int (*xTokenize)(Fts5Tokenizer*,
       void *pCtx,
       int flags,            /* Mask of FTS5_TOKENIZE_* flags */
-      const char *pText, int nText, 
+      const char *pText, int nText,
       int (*xToken)(
         void *pCtx,         /* Copy of 2nd argument to xTokenize() */
         int tflags,         /* Mask of FTS5_TOKEN_* flags */
@@ -12023,6 +13502,7 @@ struct fts5_tokenizer {
   );
 };
 
+
 /* Flags that may be passed as the third argument to xTokenize() */
 #define FTS5_TOKENIZE_QUERY     0x0001
 #define FTS5_TOKENIZE_PREFIX    0x0002
@@ -12042,13 +13522,13 @@ struct fts5_tokenizer {
 */
 typedef struct fts5_api fts5_api;
 struct fts5_api {
-  int iVersion;                   /* Currently always set to 2 */
+  int iVersion;                   /* Currently always set to 3 */
 
   /* Create a new tokenizer */
   int (*xCreateTokenizer)(
     fts5_api *pApi,
     const char *zName,
-    void *pContext,
+    void *pUserData,
     fts5_tokenizer *pTokenizer,
     void (*xDestroy)(void*)
   );
@@ -12057,7 +13537,7 @@ struct fts5_api {
   int (*xFindTokenizer)(
     fts5_api *pApi,
     const char *zName,
-    void **ppContext,
+    void **ppUserData,
     fts5_tokenizer *pTokenizer
   );
 
@@ -12065,10 +13545,29 @@ struct fts5_api {
   int (*xCreateFunction)(
     fts5_api *pApi,
     const char *zName,
-    void *pContext,
+    void *pUserData,
     fts5_extension_function xFunction,
     void (*xDestroy)(void*)
   );
+
+  /* APIs below this point are only available if iVersion>=3 */
+
+  /* Create a new tokenizer */
+  int (*xCreateTokenizer_v2)(
+    fts5_api *pApi,
+    const char *zName,
+    void *pUserData,
+    fts5_tokenizer_v2 *pTokenizer,
+    void (*xDestroy)(void*)
+  );
+
+  /* Find an existing tokenizer */
+  int (*xFindTokenizer_v2)(
+    fts5_api *pApi,
+    const char *zName,
+    void **ppUserData,
+    fts5_tokenizer_v2 **ppTokenizer
+  );
 };
 
 /*
