Fix horde issue 738724.  Define DLLEXPORT / DLLIMPORT if they are not defined prior.

diff --git a/Engine/Plugins/Runtime/Database/SQLiteCore/Source/SQLiteCore/Public/sqlite/sqlite3.h b/Engine/Plugins/Runtime/Database/SQLiteCore/Source/SQLiteCore/Public/sqlite/sqlite3.h
--- a/Engine/Plugins/Runtime/Database/SQLiteCore/Source/SQLiteCore/Public/sqlite/sqlite3.h
+++ b/Engine/Plugins/Runtime/Database/SQLiteCore/Source/SQLiteCore/Public/sqlite/sqlite3.h
@@ -90,6 +90,18 @@ extern "C" {
 # define SQLITE_SYSAPI
 #endif
 
+// END UNREAL ENGINE CHANGE
+//
+// Preflights can fail in this file when running non-unity builds.  This is
+// because DLLEXPORT is not defined but we do define it in Platform.h.  This
+// makes it consistent and avoids any preflight issues with this updated library.
+// 
+#ifndef DLLEXPORT
+	#define DLLEXPORT
+	#define DLLIMPORT
+#endif
+// END UNREAL ENGINE CHANGE
+
 /*
 ** These no-op macros are used in front of interfaces to mark those
 ** interfaces as either deprecated or experimental.  New applications
