Fixed comment

#rnx

diff --git a/Engine/Plugins/Runtime/Database/SQLiteCore/Source/SQLiteCore/Public/SQLitePreparedStatement.h b/Engine/Plugins/Runtime/Database/SQLiteCore/Source/SQLiteCore/Public/SQLitePreparedStatement.h
--- a/Engine/Plugins/Runtime/Database/SQLiteCore/Source/SQLiteCore/Public/SQLitePreparedStatement.h
+++ b/Engine/Plugins/Runtime/Database/SQLiteCore/Source/SQLiteCore/Public/SQLitePreparedStatement.h
@@ -213,7 +213,7 @@ public:
 
 	/**
 	 * Step the SQLite prepared statement to try and move on to the next result from the statement.
-	 * @param OutResultCode Optional value to store the low-level error code of the step (@see FSQLiteDatabase::GetErrorForCode).
+	 * @param OutErrorCode Optional value to store the low-level error code of the step (@see FSQLiteDatabase::GetErrorForCode).
 	 * @note See FSQLiteDatabase::Execute for a simple example of stepping a statement.
 	 */
 	UE_API ESQLitePreparedStatementStepResult Step(int32* OutErrorCode = nullptr);
